Moved enums that were only used for NetConnections from CoreNet.h to NetConnection.h. Allowed InitConnectiont to set the max packet size. Any specified max packet sizes greater than the specified MAX_PACKET_SIZE will default to MAX_PACKET_SIZE.

diff --git a/Engine/Plugins/NetcodeUnitTest/NetcodeUnitTest/Source/NetcodeUnitTest/Classes/Net/UnitTestNetConnection.h b/Engine/Plugins/NetcodeUnitTest/NetcodeUnitTest/Source/NetcodeUnitTest/Classes/Net/UnitTestNetConnection.h
--- a/Engine/Plugins/NetcodeUnitTest/NetcodeUnitTest/Source/NetcodeUnitTest/Classes/Net/UnitTestNetConnection.h
+++ b/Engine/Plugins/NetcodeUnitTest/NetcodeUnitTest/Source/NetcodeUnitTest/Classes/Net/UnitTestNetConnection.h
@@ -19,7 +19,7 @@ class UUnitTestNetConnection : public UIpConnection
 	virtual void InitBase(UNetDriver* InDriver, class FSocket* InSocket, const FURL& InURL, EConnectionState InState,
 							int32 InMaxPacket=0, int32 InPacketOverhead=0) override;
 
-	virtual void InitConnection(UNetDriver* InDriver, EConnectionState InState, const FURL& InURL, int32 InConnectionSpeed=0) override;
+	virtual void InitConnection(UNetDriver* InDriver, EConnectionState InState, const FURL& InURL, int32 InConnectionSpeed=0, int32 InMaxPacket=0) override;
 
 
 	virtual void LowLevelSend(void* Data, int32 Count) override;
@@ -46,13 +46,12 @@ class UUnitTestNetConnection : public UIpConnection
 	DECLARE_DELEGATE_TwoParams(FReceivedRawPacketDel, void* /*Data*/, int32& /*Count*/);
 
 	/**
-	 * Delegate for notifying on (and optionally blocking) replicated actor creation
+	 * Delegate for notifying on (and optionally blocking) actor channel creation
 	 *
 	 * @param ActorClass	The class of the actor being replicated
-	 * @param bActorChannel	Whether or not this actor creation is from an actor channel
-	 * @return				Whether or not to allow creation of the actor
+	 * @return				Whether or not to allow creation of the actor channel
 	 */
-	DECLARE_DELEGATE_RetVal_TwoParams(bool, FOnReplicatedActorSpawn, UClass* /*ActorClass*/, bool /*bActorChannel*/);
+	DECLARE_DELEGATE_RetVal_OneParam(bool, FOnActorChannelSpawn, UClass* /*ActorClass*/);
 
 
 	/** Delegate for hooking LowLevelSend */
@@ -61,8 +60,8 @@ class UUnitTestNetConnection : public UIpConnection
 	/** Delegate for hooking ReceivedRawPacket */
 	FReceivedRawPacketDel	ReceivedRawPacketDel;
 
-	/** Delegate for notifying on replicated actor creation */
-	FOnReplicatedActorSpawn	ReplicatedActorSpawnDel;
+	/** Delegate for notifying on actor channel creation */
+	FOnActorChannelSpawn	ActorChannelSpawnDel;
 };
 
 
diff --git a/Engine/Plugins/NetcodeUnitTest/NetcodeUnitTest/Source/NetcodeUnitTest/Private/Net/UnitTestNetConnection.cpp b/Engine/Plugins/NetcodeUnitTest/NetcodeUnitTest/Source/NetcodeUnitTest/Private/Net/UnitTestNetConnection.cpp
--- a/Engine/Plugins/NetcodeUnitTest/NetcodeUnitTest/Source/NetcodeUnitTest/Private/Net/UnitTestNetConnection.cpp
+++ b/Engine/Plugins/NetcodeUnitTest/NetcodeUnitTest/Source/NetcodeUnitTest/Private/Net/UnitTestNetConnection.cpp
@@ -38,7 +38,7 @@ void UUnitTestNetConnection::InitBase(UNetDriver* InDriver, class FSocket* InSoc
 }
 
 void UUnitTestNetConnection::InitConnection(UNetDriver* InDriver, EConnectionState InState, const FURL& InURL,
-											int32 InConnectionSpeed/*=0*/)
+											int32 InConnectionSpeed, int32 InMaxPacket)
 {
 	Super::InitConnection(InDriver, InState, InURL, InConnectionSpeed);
 
diff --git a/Engine/Source/Runtime/CoreUObject/Public/UObject/CoreNet.h b/Engine/Source/Runtime/CoreUObject/Public/UObject/CoreNet.h
--- a/Engine/Source/Runtime/CoreUObject/Public/UObject/CoreNet.h
+++ b/Engine/Source/Runtime/CoreUObject/Public/UObject/CoreNet.h
@@ -433,13 +433,8 @@ COREUOBJECT_API void SerializeChecksum(FArchive &Ar, uint32 x, bool ErrorOK);
 /**
 * Values used for initializing UNetConnection and LanBeacon
 */
-enum { RELIABLE_BUFFER = 256 }; // Power of 2 >= 1.
-enum { MAX_PACKETID = 16384 };  // Power of 2 >= 1, covering guaranteed loss/misorder time.
-enum { MAX_CHSEQUENCE = 1024 }; // Power of 2 >RELIABLE_BUFFER, covering loss/misorder time.
-enum { MAX_BUNCH_HEADER_BITS = 64 };
-enum { MAX_PACKET_HEADER_BITS = 16 };
-enum { MAX_PACKET_TRAILER_BITS = 1 };
 enum { MAX_PACKET_SIZE = 512 }; // MTU for the connection
+enum { LAN_BEACON_MAX_PACKET_SIZE = 512 }; // MTU for the connection
 
 /**
  * Functions to assist in detecting errors during RPC calls
diff --git a/Engine/Source/Runtime/Engine/Classes/Engine/DemoNetConnection.h b/Engine/Source/Runtime/Engine/Classes/Engine/DemoNetConnection.h
--- a/Engine/Source/Runtime/Engine/Classes/Engine/DemoNetConnection.h
+++ b/Engine/Source/Runtime/Engine/Classes/Engine/DemoNetConnection.h
@@ -24,7 +24,7 @@ public:
 
 	// UNetConnection interface.
 
-	virtual void InitConnection( class UNetDriver* InDriver, EConnectionState InState, const FURL& InURL, int32 InConnectionSpeed = 0 ) override;
+	virtual void InitConnection( class UNetDriver* InDriver, EConnectionState InState, const FURL& InURL, int32 InConnectionSpeed = 0, int32 InMaxPacket=0) override;
 	virtual FString LowLevelGetRemoteAddress( bool bAppendPort = false ) override;
 	virtual FString LowLevelDescribe() override;
 	virtual void LowLevelSend( void* Data, int32 Count ) override;
diff --git a/Engine/Source/Runtime/Engine/Classes/Engine/NetConnection.h b/Engine/Source/Runtime/Engine/Classes/Engine/NetConnection.h
--- a/Engine/Source/Runtime/Engine/Classes/Engine/NetConnection.h
+++ b/Engine/Source/Runtime/Engine/Classes/Engine/NetConnection.h
@@ -19,7 +19,12 @@ class FUniqueNetId;
 /*-----------------------------------------------------------------------------
 	Types.
 -----------------------------------------------------------------------------*/
-
+enum { RELIABLE_BUFFER = 256 }; // Power of 2 >= 1.
+enum { MAX_PACKETID = 16384 };  // Power of 2 >= 1, covering guaranteed loss/misorder time.
+enum { MAX_CHSEQUENCE = 1024 }; // Power of 2 >RELIABLE_BUFFER, covering loss/misorder time.
+enum { MAX_BUNCH_HEADER_BITS = 64 };
+enum { MAX_PACKET_HEADER_BITS = 16 };
+enum { MAX_PACKET_TRAILER_BITS = 1 };
 
 class UNetDriver;
 
@@ -475,7 +480,7 @@ public:
 	 * @param InURL the URL to init with
 	 * @param InConnectionSpeed Optional connection speed override
 	 */
-	ENGINE_API virtual void InitConnection(UNetDriver* InDriver, EConnectionState InState, const FURL& InURL, int32 InConnectionSpeed=0);
+	ENGINE_API virtual void InitConnection(UNetDriver* InDriver, EConnectionState InState, const FURL& InURL, int32 InConnectionSpeed=0, int32 InMaxPacket=0);
 
 	// Functions.
 
diff --git a/Engine/Source/Runtime/Engine/Private/DemoNetDriver.cpp b/Engine/Source/Runtime/Engine/Private/DemoNetDriver.cpp
--- a/Engine/Source/Runtime/Engine/Private/DemoNetDriver.cpp
+++ b/Engine/Source/Runtime/Engine/Private/DemoNetDriver.cpp
@@ -1715,12 +1715,12 @@ UDemoNetConnection::UDemoNetConnection( const FObjectInitializer& ObjectInitiali
 	InternalAck = true;
 }
 
-void UDemoNetConnection::InitConnection( UNetDriver* InDriver, EConnectionState InState, const FURL& InURL, int32 InConnectionSpeed )
+void UDemoNetConnection::InitConnection( UNetDriver* InDriver, EConnectionState InState, const FURL& InURL, int32 InConnectionSpeed, int32 InMaxPacket)
 {
 	// default implementation
 	Super::InitConnection( InDriver, InState, InURL, InConnectionSpeed );
 
-	MaxPacket = MAX_DEMO_READ_WRITE_BUFFER;
+	MaxPacket = (InMaxPacket == 0 || InMaxPacket > MAX_PACKET_SIZE) ? MAX_DEMO_READ_WRITE_BUFFER : InMaxPacket;
 	InternalAck = true;
 
 	InitSendBuffer();
diff --git a/Engine/Source/Runtime/Engine/Private/NetConnection.cpp b/Engine/Source/Runtime/Engine/Private/NetConnection.cpp
--- a/Engine/Source/Runtime/Engine/Private/NetConnection.cpp
+++ b/Engine/Source/Runtime/Engine/Private/NetConnection.cpp
@@ -133,11 +133,11 @@ void UNetConnection::InitBase(UNetDriver* InDriver,class FSocket* InSocket, cons
  * @param InURL the URL to init with
  * @param InConnectionSpeed Optional connection speed override
  */
-void UNetConnection::InitConnection(UNetDriver* InDriver, EConnectionState InState, const FURL& InURL, int32 InConnectionSpeed)
+void UNetConnection::InitConnection(UNetDriver* InDriver, EConnectionState InState, const FURL& InURL, int32 InConnectionSpeed, int32 InMaxPacket)
 {
 	Driver = InDriver;
 	// We won't be sending any packets, so use a default size
-	MaxPacket = MAX_PACKET_SIZE;
+	MaxPacket = (InMaxPacket == 0 || InMaxPacket > MAX_PACKET_SIZE) ? MAX_PACKET_SIZE : InMaxPacket;
 	PacketOverhead = 0;
 	State = InState;
 
diff --git a/Engine/Source/Runtime/Online/OnlineSubsystem/Public/LANBeacon.h b/Engine/Source/Runtime/Online/OnlineSubsystem/Public/LANBeacon.h
--- a/Engine/Source/Runtime/Online/OnlineSubsystem/Public/LANBeacon.h
+++ b/Engine/Source/Runtime/Online/OnlineSubsystem/Public/LANBeacon.h
@@ -16,9 +16,6 @@
 
 /** The size of the header for validation */
 #define LAN_BEACON_PACKET_HEADER_SIZE 16
-
-/** The max size expected in a lan beacon packet */
-#define LAN_BEACON_MAX_PACKET_SIZE MAX_PACKET_SIZE
 	
 // Offsets for various fields
 #define LAN_BEACON_VER_OFFSET 0
diff --git a/Engine/Source/Runtime/Online/OnlineSubsystemUtils/Private/IpConnection.cpp b/Engine/Source/Runtime/Online/OnlineSubsystemUtils/Private/IpConnection.cpp
--- a/Engine/Source/Runtime/Online/OnlineSubsystemUtils/Private/IpConnection.cpp
+++ b/Engine/Source/Runtime/Online/OnlineSubsystemUtils/Private/IpConnection.cpp
@@ -36,7 +36,7 @@ void UIpConnection::InitBase(UNetDriver* InDriver, class FSocket* InSocket, cons
 	// Pass the call up the chain
 	Super::InitBase(InDriver, InSocket, InURL, InState, 
 		// Use the default packet size/overhead unless overridden by a child class
-		InMaxPacket == 0 ? MAX_PACKET_SIZE : InMaxPacket,
+		(InMaxPacket == 0 || InMaxPacket > MAX_PACKET_SIZE) ? MAX_PACKET_SIZE : InMaxPacket,
 		InPacketOverhead == 0 ? SLIP_HEADER_SIZE : InPacketOverhead);
 
 	Socket = InSocket;
@@ -47,7 +47,7 @@ void UIpConnection::InitLocalConnection(UNetDriver* InDriver, class FSocket* InS
 {
 	InitBase(InDriver, InSocket, InURL, InState, 
 		// Use the default packet size/overhead unless overridden by a child class
-		InMaxPacket == 0 ? MAX_PACKET_SIZE : InMaxPacket,
+		(InMaxPacket == 0 || InMaxPacket > MAX_PACKET_SIZE) ? MAX_PACKET_SIZE : InMaxPacket,
 		InPacketOverhead == 0 ? SLIP_HEADER_SIZE : InPacketOverhead);
 
 	// Figure out IP address from the host URL
@@ -77,7 +77,7 @@ void UIpConnection::InitRemoteConnection(UNetDriver* InDriver, class FSocket* In
 {
 	InitBase(InDriver, InSocket, InURL, InState, 
 		// Use the default packet size/overhead unless overridden by a child class
-		InMaxPacket == 0 ? MAX_PACKET_SIZE : InMaxPacket,
+		(InMaxPacket == 0 || InMaxPacket > MAX_PACKET_SIZE) ? MAX_PACKET_SIZE : InMaxPacket,
 		InPacketOverhead == 0 ? SLIP_HEADER_SIZE : InPacketOverhead);
 
 	// Copy the remote IPAddress passed in
