Fix for previous submit

diff --git a/Engine/Plugins/NetcodeUnitTest/NetcodeUnitTest/Source/NetcodeUnitTest/Classes/Net/UnitTestNetConnection.h b/Engine/Plugins/NetcodeUnitTest/NetcodeUnitTest/Source/NetcodeUnitTest/Classes/Net/UnitTestNetConnection.h
--- a/Engine/Plugins/NetcodeUnitTest/NetcodeUnitTest/Source/NetcodeUnitTest/Classes/Net/UnitTestNetConnection.h
+++ b/Engine/Plugins/NetcodeUnitTest/NetcodeUnitTest/Source/NetcodeUnitTest/Classes/Net/UnitTestNetConnection.h
@@ -8,8 +8,8 @@
 
 
 /**
- * A net connection for enabling unit testing through barebones/minimal client connections
- */
+* A net connection for enabling unit testing through barebones/minimal client connections
+*/
 UCLASS(transient)
 class UUnitTestNetConnection : public UIpConnection
 {
@@ -17,9 +17,9 @@ class UUnitTestNetConnection : public UIpConnection
 
 
 	virtual void InitBase(UNetDriver* InDriver, class FSocket* InSocket, const FURL& InURL, EConnectionState InState,
-							int32 InMaxPacket=0, int32 InPacketOverhead=0) override;
+	int32 InMaxPacket = 0, int32 InPacketOverhead = 0) override;
 
-	virtual void InitConnection(UNetDriver* InDriver, EConnectionState InState, const FURL& InURL, int32 InConnectionSpeed=0, int32 InMaxPacket=0) override;
+	virtual void InitConnection(UNetDriver* InDriver, EConnectionState InState, const FURL& InURL, int32 InConnectionSpeed = 0, int32 InMaxPacket=0) override;
 
 
 	virtual void LowLevelSend(void* Data, int32 Count) override;
@@ -30,28 +30,29 @@ class UUnitTestNetConnection : public UIpConnection
 
 
 	/**
-	 * Delegate for hooking 'LowLevelSend'
-	 *
-	 * @param Data		The data being sent
-	 * @param Count		The number of bytes being sent
-	 */
+	* Delegate for hooking 'LowLevelSend'
+	*
+	* @param Data		The data being sent
+	* @param Count		The number of bytes being sent
+	*/
 	DECLARE_DELEGATE_TwoParams(FLowLevelSendDel, void* /*Data*/, int32 /*Count*/);
 
 	/**
-	 * Delegate for hooking 'ReceivedRawPacket'
-	 * 
-	 * @param Data		The data received
-	 * @param Count		The number of bytes received
-	 */
+	* Delegate for hooking 'ReceivedRawPacket'
+	*
+	* @param Data		The data received
+	* @param Count		The number of bytes received
+	*/
 	DECLARE_DELEGATE_TwoParams(FReceivedRawPacketDel, void* /*Data*/, int32& /*Count*/);
 
 	/**
-	 * Delegate for notifying on (and optionally blocking) actor channel creation
-	 *
-	 * @param ActorClass	The class of the actor being replicated
-	 * @return				Whether or not to allow creation of the actor channel
-	 */
-	DECLARE_DELEGATE_RetVal_OneParam(bool, FOnActorChannelSpawn, UClass* /*ActorClass*/);
+	* Delegate for notifying on (and optionally blocking) replicated actor creation
+	*
+	* @param ActorClass	The class of the actor being replicated
+	* @param bActorChannel	Whether or not this actor creation is from an actor channel
+	* @return				Whether or not to allow creation of the actor
+	*/
+	DECLARE_DELEGATE_RetVal_TwoParams(bool, FOnReplicatedActorSpawn, UClass* /*ActorClass*/, bool /*bActorChannel*/);
 
 
 	/** Delegate for hooking LowLevelSend */
@@ -60,8 +61,8 @@ class UUnitTestNetConnection : public UIpConnection
 	/** Delegate for hooking ReceivedRawPacket */
 	FReceivedRawPacketDel	ReceivedRawPacketDel;
 
-	/** Delegate for notifying on actor channel creation */
-	FOnActorChannelSpawn	ActorChannelSpawnDel;
+	/** Delegate for notifying on replicated actor creation */
+	FOnReplicatedActorSpawn	ReplicatedActorSpawnDel;
 };
 
 
