Shadow variable fixes

diff --git a/Engine/Plugins/NetcodeUnitTest/NetcodeUnitTest/Source/NetcodeUnitTest/Private/NUTActor.cpp b/Engine/Plugins/NetcodeUnitTest/NetcodeUnitTest/Source/NetcodeUnitTest/Private/NUTActor.cpp
--- a/Engine/Plugins/NetcodeUnitTest/NetcodeUnitTest/Source/NetcodeUnitTest/Private/NUTActor.cpp
+++ b/Engine/Plugins/NetcodeUnitTest/NetcodeUnitTest/Source/NetcodeUnitTest/Private/NUTActor.cpp
@@ -133,10 +133,10 @@ bool ANUTActor::NotifyControlMessage(UNetConnection* Connection, uint8 MessageTy
 		// Ping request
 		else if (CmdType == ENUTControlCommand::Ping)
 		{
-			uint8 CmdType = ENUTControlCommand::Pong;
+			uint8 TempCmdType = ENUTControlCommand::Pong;
 			FString Dud;
 
-			FNetControlMessage<NMT_NUTControl>::Send(Connection, CmdType, Dud);
+			FNetControlMessage<NMT_NUTControl>::Send(Connection, TempCmdType, Dud);
 		}
 		// Pong reply - this should only be implemented by custom unit tests; hence the assert
 		else if (CmdType == ENUTControlCommand::Pong)
