Compatability fixes and minor changes.

Full changes:
- Fixed error reading the minor engine version, which broke the netcode in the transition from 4.9 to 4.10
- Added '-UnitTestClientDebug' commandline parameter, which launches unit test clients with a valid render interface, to allow interaction when debugging
- Added a 'MaxAutoCloseCount' config variable, to limit the number of times a unit test can be auto-aborted, before the unit test tool refuses to run it again

- Fixed rare infinite recursion in unit test log window search
- Added unit test failure condition, when trying to send an RPC when the net buffers are saturated
- Increased the wait time to 2 minutes before auto-aborting, when a unit test crashes, due to long wait time for crash stack dump sometimes
- Added better notification of unit tests aborts, in the status window
- Prevent net connections timing out in developer mode
- Updated comments to use VAX hashtags

diff --git a/Engine/Plugins/NetcodeUnitTest/NetcodeUnitTest/Source/NetcodeUnitTest/Classes/ClientUnitTest.h b/Engine/Plugins/NetcodeUnitTest/NetcodeUnitTest/Source/NetcodeUnitTest/Classes/ClientUnitTest.h
--- a/Engine/Plugins/NetcodeUnitTest/NetcodeUnitTest/Source/NetcodeUnitTest/Classes/ClientUnitTest.h
+++ b/Engine/Plugins/NetcodeUnitTest/NetcodeUnitTest/Source/NetcodeUnitTest/Classes/ClientUnitTest.h
@@ -364,9 +364,11 @@ public:
 	FOnSuspendStateChange OnServerSuspendState;
 
 
+#if TARGET_UE4_CL >= CL_DEPRECATEDEL
 private:
 	/** Handle to the registered InternalNotifyNetworkFailure delegate */
 	FDelegateHandle InternalNotifyNetworkFailureDelegateHandle;
+#endif
 
 
 	/**
@@ -605,6 +607,11 @@ protected:
 
 	virtual void ResetTimeout(FString ResetReason, bool bResetConnTimeout=false, uint32 MinDuration=0) override;
 
+	/**
+	 * Resets the net connection timeout
+	 */
+	void ResetConnTimeout();
+
 
 	/**
 	 * Returns the requirements flags, that this unit test currently meets
diff --git a/Engine/Plugins/NetcodeUnitTest/NetcodeUnitTest/Source/NetcodeUnitTest/Classes/NUTActor.h b/Engine/Plugins/NetcodeUnitTest/NetcodeUnitTest/Source/NetcodeUnitTest/Classes/NUTActor.h
--- a/Engine/Plugins/NetcodeUnitTest/NetcodeUnitTest/Source/NetcodeUnitTest/Classes/NUTActor.h
+++ b/Engine/Plugins/NetcodeUnitTest/NetcodeUnitTest/Source/NetcodeUnitTest/Classes/NUTActor.h
@@ -14,7 +14,7 @@ class ANUTActor;
 /**
  * Enum for defining custom NetcodeUnitTest control channel commands (sent through NMT_NUTControl)
  */
-// @todo JohnB: Fully document individually
+// @todo #JohnBDoc: Fully document individually
 namespace ENUTControlCommand
 {
 	enum Type
@@ -162,7 +162,7 @@ public:
 	/**
 	 * Received by all clients, emits a ping to log
 	 */
-	// @todo JohnB: When the VM reflection helper is finished, remove NETCODEUNITTEST_API from this, and use reflection instead
+	// @todo #JohnBRefactor: When the VM reflection helper is finished, remove NETCODEUNITTEST_API from this, and use reflection instead
 	UFUNCTION(reliable, NetMulticast)
 	NETCODEUNITTEST_API void NetMulticastPing();
 
diff --git a/Engine/Plugins/NetcodeUnitTest/NetcodeUnitTest/Source/NetcodeUnitTest/Classes/Net/UnitTestNetConnection.h b/Engine/Plugins/NetcodeUnitTest/NetcodeUnitTest/Source/NetcodeUnitTest/Classes/Net/UnitTestNetConnection.h
--- a/Engine/Plugins/NetcodeUnitTest/NetcodeUnitTest/Source/NetcodeUnitTest/Classes/Net/UnitTestNetConnection.h
+++ b/Engine/Plugins/NetcodeUnitTest/NetcodeUnitTest/Source/NetcodeUnitTest/Classes/Net/UnitTestNetConnection.h
@@ -8,18 +8,23 @@
 
 
 /**
-* A net connection for enabling unit testing through barebones/minimal client connections
-*/
+ * A net connection for enabling unit testing through barebones/minimal client connections
+ */
 UCLASS(transient)
 class UUnitTestNetConnection : public UIpConnection
 {
 	GENERATED_UCLASS_BODY()
 
-	 
+
 	virtual void InitBase(UNetDriver* InDriver, class FSocket* InSocket, const FURL& InURL, EConnectionState InState,
-	int32 InMaxPacket = 0, int32 InPacketOverhead = 0) override;
+							int32 InMaxPacket=0, int32 InPacketOverhead=0) override;
 
-	virtual void InitConnection(UNetDriver* InDriver, EConnectionState InState, const FURL& InURL, int32 InConnectionSpeed = 0, int32 InMaxPacket = 0) override;
+#if TARGET_UE4_CL < CL_INITCONNPARAM
+	virtual void InitConnection(UNetDriver* InDriver, EConnectionState InState, const FURL& InURL, int32 InConnectionSpeed=0) override;
+#else
+	virtual void InitConnection(UNetDriver* InDriver, EConnectionState InState, const FURL& InURL, int32 InConnectionSpeed=0,
+								int32 InMaxPacket=0) override;
+#endif
 
 
 	virtual void LowLevelSend(void* Data, int32 Count) override;
@@ -30,28 +35,28 @@ class UUnitTestNetConnection : public UIpConnection
 
 
 	/**
-	* Delegate for hooking 'LowLevelSend'
-	*
-	* @param Data		The data being sent
-	* @param Count		The number of bytes being sent
-	*/
+	 * Delegate for hooking 'LowLevelSend'
+	 *
+	 * @param Data		The data being sent
+	 * @param Count		The number of bytes being sent
+	 */
 	DECLARE_DELEGATE_TwoParams(FLowLevelSendDel, void* /*Data*/, int32 /*Count*/);
 
 	/**
-	* Delegate for hooking 'ReceivedRawPacket'
-	*
-	* @param Data		The data received
-	* @param Count		The number of bytes received
-	*/
+	 * Delegate for hooking 'ReceivedRawPacket'
+	 * 
+	 * @param Data		The data received
+	 * @param Count		The number of bytes received
+	 */
 	DECLARE_DELEGATE_TwoParams(FReceivedRawPacketDel, void* /*Data*/, int32& /*Count*/);
 
 	/**
-	* Delegate for notifying on (and optionally blocking) replicated actor creation
-	*
-	* @param ActorClass	The class of the actor being replicated
-	* @param bActorChannel	Whether or not this actor creation is from an actor channel
-	* @return				Whether or not to allow creation of the actor
-	*/
+	 * Delegate for notifying on (and optionally blocking) replicated actor creation
+	 *
+	 * @param ActorClass	The class of the actor being replicated
+	 * @param bActorChannel	Whether or not this actor creation is from an actor channel
+	 * @return				Whether or not to allow creation of the actor
+	 */
 	DECLARE_DELEGATE_RetVal_TwoParams(bool, FOnReplicatedActorSpawn, UClass* /*ActorClass*/, bool /*bActorChannel*/);
 
 
@@ -68,3 +73,4 @@ class UUnitTestNetConnection : public UIpConnection
 
 
 
+
diff --git a/Engine/Plugins/NetcodeUnitTest/NetcodeUnitTest/Source/NetcodeUnitTest/Classes/UnitTest.h b/Engine/Plugins/NetcodeUnitTest/NetcodeUnitTest/Source/NetcodeUnitTest/Classes/UnitTest.h
--- a/Engine/Plugins/NetcodeUnitTest/NetcodeUnitTest/Source/NetcodeUnitTest/Classes/UnitTest.h
+++ b/Engine/Plugins/NetcodeUnitTest/NetcodeUnitTest/Source/NetcodeUnitTest/Classes/UnitTest.h
@@ -7,7 +7,7 @@
 #include "UnitTest.generated.h"
 
 
-// @todo JohnB: For bugtracking/changelist info, consider adding auto-launching of P4/TTP/Browser-JIRA links,
+// @todo #JohnBFeature: For bugtracking/changelist info, consider adding auto-launching of P4/TTP/Browser-JIRA links,
 //				upon double-clicking these entries in the status windows
 
 
@@ -173,7 +173,7 @@ protected:
 	bool bFirstTimeStats;
 
 
-	// @todo JohnB: Merge the two below variables
+	// @todo #JohnBRefactor: Merge the two below variables
 	/** Whether or not the unit test has completed */
 	bool bCompleted;
 
diff --git a/Engine/Plugins/NetcodeUnitTest/NetcodeUnitTest/Source/NetcodeUnitTest/Classes/UnitTestCommandlet.h b/Engine/Plugins/NetcodeUnitTest/NetcodeUnitTest/Source/NetcodeUnitTest/Classes/UnitTestCommandlet.h
--- a/Engine/Plugins/NetcodeUnitTest/NetcodeUnitTest/Source/NetcodeUnitTest/Classes/UnitTestCommandlet.h
+++ b/Engine/Plugins/NetcodeUnitTest/NetcodeUnitTest/Source/NetcodeUnitTest/Classes/UnitTestCommandlet.h
@@ -26,6 +26,9 @@
  *	-UnitTestClientParms="CommandlineParameters"
  *		- Adds additional commandline parameters to unit test client instances
  *
+ *	-UnitTestClientDebug
+ *		- Launches unit test clients with a valid render interface (so you can interact with the client window), and a larger window
+ *
  *	-UnitTestCap=x
  *		- Caps the maximum number of unit tests that can run at the same time
  */
diff --git a/Engine/Plugins/NetcodeUnitTest/NetcodeUnitTest/Source/NetcodeUnitTest/Classes/UnitTestManager.h b/Engine/Plugins/NetcodeUnitTest/NetcodeUnitTest/Source/NetcodeUnitTest/Classes/UnitTestManager.h
--- a/Engine/Plugins/NetcodeUnitTest/NetcodeUnitTest/Source/NetcodeUnitTest/Classes/UnitTestManager.h
+++ b/Engine/Plugins/NetcodeUnitTest/NetcodeUnitTest/Source/NetcodeUnitTest/Classes/UnitTestManager.h
@@ -35,7 +35,7 @@ public:
 	UPROPERTY(config)
 	bool bCapUnitTestMemory;
 
-	// @todo JohnB: Add a bool, for specifying that unit tests can be terminated, when memory limits are breached
+	// @todo #JohnBLowPri: Add a bool, for specifying that unit tests can be terminated, when memory limits are breached
 
 	/** When total physical memory usage, as a percentage, reaches this limit, no new unit tests can be started */
 	UPROPERTY(config)
@@ -45,6 +45,10 @@ public:
 	UPROPERTY(config)
 	uint8 AutoCloseMemoryPercent;
 
+	/** Limits the number of auto-aborts a particular unit test will allow, before it is no longer accepted for re-queueing */
+	UPROPERTY(config)
+	uint8 MaxAutoCloseCount;
+
 
 	/** Holds a list of unit tests pending execution */
 	UPROPERTY()
@@ -88,6 +92,9 @@ public:
 
 
 private:
+	/** The time at which the memory limit was last hit */
+	double LastMemoryLimitHit;
+
 	/** When a unit test is force-closed, wait a number of ticks for global memory values to update, before closing any more */
 	int32 MemoryTickCountdown;
 
@@ -201,8 +208,9 @@ public:
 	 *
 	 * @param InUnitTest		The unit test to print results information for
 	 * @param bFinalSummary		Whether or not this is the final summary printout (changes the formatting slightly)
+	 * @param bUnfinished		Whether or not the unit test was aborted and could not be run
 	 */
-	void PrintUnitTestResult(UUnitTest* InUnitTest, bool bFinalSummary=false);
+	void PrintUnitTestResult(UUnitTest* InUnitTest, bool bFinalSummary=false, bool bUnfinished=false);
 
 	/**
 	 * Prints the final unit test summary, when all active/pending unit tests have completed
@@ -232,7 +240,7 @@ public:
 	// Must override in subclasses, that need ticking
 	virtual bool IsTickable() const override
 	{
-		// @todo JohnB: Find out how the CDO is getting registered for ticking - this is odd
+		// @todo #JohnBLowPri: Find out how the CDO is getting registered for ticking - this is odd
 		return !IsPendingKill() && !GIsServer && !HasAnyFlags(RF_ClassDefaultObject);
 	}
 
diff --git a/Engine/Plugins/NetcodeUnitTest/NetcodeUnitTest/Source/NetcodeUnitTest/NetcodeUnitTest.Build.cs b/Engine/Plugins/NetcodeUnitTest/NetcodeUnitTest/Source/NetcodeUnitTest/NetcodeUnitTest.Build.cs
--- a/Engine/Plugins/NetcodeUnitTest/NetcodeUnitTest/Source/NetcodeUnitTest/NetcodeUnitTest.Build.cs
+++ b/Engine/Plugins/NetcodeUnitTest/NetcodeUnitTest/Source/NetcodeUnitTest/NetcodeUnitTest.Build.cs
@@ -24,7 +24,7 @@ namespace UnrealBuildTool.Rules
 				}
 			);
 
-			// @todo JohnB: Currently don't support standalone commandlet, with static builds (can't get past linker error in Win32)
+			// @todo #JohnBLowpri: Currently don't support standalone commandlet, with static builds (can't get past linker error in Win32)
 			if (!Target.IsMonolithic)
 			{
 				PrivateDependencyModuleNames.AddRange
diff --git a/Engine/Plugins/NetcodeUnitTest/NetcodeUnitTest/Source/NetcodeUnitTest/Private/ClientUnitTest.cpp b/Engine/Plugins/NetcodeUnitTest/NetcodeUnitTest/Source/NetcodeUnitTest/Private/ClientUnitTest.cpp
--- a/Engine/Plugins/NetcodeUnitTest/NetcodeUnitTest/Source/NetcodeUnitTest/Private/ClientUnitTest.cpp
+++ b/Engine/Plugins/NetcodeUnitTest/NetcodeUnitTest/Source/NetcodeUnitTest/Private/ClientUnitTest.cpp
@@ -25,6 +25,11 @@
 #include "Regex.h"
 
 
+// @todo #JohnBMultiFakeClient: Eventually, move >all< of the minimal/headless client handling code, into a new/separate class,
+//				so that a single unit test can have multiple minimal clients on a server.
+//				This would be useful for licensees, for doing load testing:
+//				https://udn.unrealengine.com/questions/247014/clientserver-automation.html
+
 /**
  * UClientUnitTest
  */
@@ -498,7 +503,7 @@ void UClientUnitTest::NotifyProcessLog(TWeakPtr<FUnitTestProcess> InProcess, con
 		}
 	}
 
-	// @todo JohnB: Consider also, adding a way to communicate with launched clients,
+	// @todo #JohnBLowPri: Consider also, adding a way to communicate with launched clients,
 	//				to reset their connection timeout upon server progress, if they fully startup before the server does
 }
 
@@ -586,7 +591,7 @@ void UClientUnitTest::NotifyProcessSuspendState(TWeakPtr<FUnitTestProcess> InPro
 		OnServerSuspendState.ExecuteIfBound(InSuspendState);
 	}
 
-	// @todo JohnB: Want to support client process debugging in future?
+	// @todo #JohnBLowPri: Want to support client process debugging in future?
 }
 
 
@@ -610,8 +615,8 @@ bool UClientUnitTest::NotifyConsoleCommandRequest(FString CommandContext, FStrin
 
 		if (!BadCmds.Contains(Command))
 		{
-			// @todo JohnB: Should this mark the log origin, as from the unit test?
-			// @todo JohnB: In general, I'm not sure how I handle the log-origin of UI-triggered events;
+			// @todo #JohnB: Should this mark the log origin, as from the unit test?
+			// @todo #JohnB: In general, I'm not sure how I handle the log-origin of UI-triggered events;
 			//				they maybe should be captured/categorized better
 			UNIT_LOG_BEGIN(this, ELogType::OriginConsole);
 			bHandled = GEngine->Exec(TargetWorld, *Command, *GLog);
@@ -625,7 +630,7 @@ bool UClientUnitTest::NotifyConsoleCommandRequest(FString CommandContext, FStrin
 	}
 	else if (CommandContext == TEXT("Server"))
 	{
-		// @todo JohnB: Perhaps add extra checks here, to be sure we're ready to send console commands?
+		// @todo #JohnBBug: Perhaps add extra checks here, to be sure we're ready to send console commands?
 		//
 		//				UPDATE: Yes, this is a good idea, because if the client hasn't gotten to the correct login stage
 		//				(NMT_Join or such, need to check when server rejects non-login control commands),
@@ -658,7 +663,7 @@ bool UClientUnitTest::NotifyConsoleCommandRequest(FString CommandContext, FStrin
 	}
 	else if (CommandContext == TEXT("Client"))
 	{
-		// @todo JohnB
+		// @todo #JohnBFeature
 
 		UNIT_LOG(ELogType::OriginConsole, TEXT("Client console commands not yet implemented"));
 	}
@@ -702,7 +707,14 @@ bool UClientUnitTest::SendRPCChecked(AActor* Target, const TCHAR* FunctionName,
 	{
 		if (TargetFunc->ParmsSize == ParmsSize + ParmsSizeCorrection)
 		{
-			Target->ProcessEvent(TargetFunc, Parms);
+			if (UnitConn->IsNetReady(false))
+			{
+				Target->ProcessEvent(TargetFunc, Parms);
+			}
+			else
+			{
+				UNIT_LOG(ELogType::StatusFailure, TEXT("Failed to send RPC '%s', network saturated."), FunctionName);
+			}
 		}
 		else
 		{
@@ -765,7 +777,7 @@ bool UClientUnitTest::PostSendRPC(FString RPCName, AActor* Target/*=NULL*/)
 }
 
 
-// @todo JohnB: This should be changed to work as compile-time conditionals, if possible
+// @todo #JohnBRefactor: This should be changed to work as compile-time conditionals, if possible
 bool UClientUnitTest::ValidateUnitTestSettings(bool bCDOCheck/*=false*/)
 {
 	bool bSuccess = Super::ValidateUnitTestSettings();
@@ -973,9 +985,17 @@ void UClientUnitTest::ResetTimeout(FString ResetReason, bool bResetConnTimeout/*
 
 	Super::ResetTimeout(ResetReason, bResetConnTimeout, MinDuration);
 
-	if (bResetConnTimeout && UnitConn != NULL && UnitConn->State != USOCK_Closed && UnitConn->Driver != NULL)
+	if (bResetConnTimeout)
+	{
+		ResetConnTimeout();
+	}
+}
+
+void UClientUnitTest::ResetConnTimeout()
+{
+	if (UnitConn != NULL && UnitConn->State != USOCK_Closed && UnitConn->Driver != NULL)
 	{
-		// @todo JohnB: This is a slightly hacky way of setting the timeout to a large value, which will be overridden by newly
+		// @todo #JohnBHack: This is a slightly hacky way of setting the timeout to a large value, which will be overridden by newly
 		//				received packets, making it unsuitable for most situations (except crashes - but that could still be subject
 		//				to a race condition)
 		UnitConn->LastReceiveTime = UnitConn->Driver->Time + (float)(TimeoutExpire - FPlatformTime::Seconds());
@@ -989,7 +1009,7 @@ bool UClientUnitTest::ExecuteUnitTest()
 
 	bool bValidSettings = ValidateUnitTestSettings();
 
-	// @todo JohnB: Fix support for Steam eventually
+	// @todo #JohnBLowPri: Fix support for Steam eventually
 	bool bSteamAvailable = NUTNet::IsSteamNetDriverAvailable();
 
 	if (bValidSettings && !bSteamAvailable)
@@ -1093,8 +1113,12 @@ bool UClientUnitTest::ConnectFakeClient(FUniqueNetIdRepl* InNetID/*=NULL*/)
 
 					if (GEngine != NULL)
 					{
+#if TARGET_UE4_CL >= CL_DEPRECATEDEL
 						InternalNotifyNetworkFailureDelegateHandle = GEngine->OnNetworkFailure().AddUObject(this,
 																		&UClientUnitTest::InternalNotifyNetworkFailure);
+#else
+						GEngine->OnNetworkFailure().AddUObject(this, &UClientUnitTest::InternalNotifyNetworkFailure);
+#endif
 					}
 
 
@@ -1189,12 +1213,12 @@ bool UClientUnitTest::ConnectFakeClient(FUniqueNetIdRepl* InNetID/*=NULL*/)
 	return bSuccess;
 }
 
-// @todo JohnB: Reconsider naming of this function, if you modify the other functions using 'fake'
+// @todo #JohnBRefactor: Reconsider naming of this function, if you modify the other functions using 'fake'
 void UClientUnitTest::CleanupFakeClient()
 {
 	if (UnitWorld != NULL && UnitNetDriver != NULL)
 	{
-		// @todo JohnB: As with the 'CreateFakePlayer' function, naming it 'fake' may not be optimal
+		// @todo #JohnBRefactor: As with the 'CreateFakePlayer' function, naming it 'fake' may not be optimal
 		NUTNet::DisconnectFakePlayer(UnitWorld, UnitNetDriver);
 
 		UnitNetDriver->Notify = NULL;
@@ -1213,7 +1237,11 @@ void UClientUnitTest::CleanupFakeClient()
 
 	if (GEngine != NULL)
 	{
+#if TARGET_UE4_CL >= CL_DEPRECATEDEL
 		GEngine->OnNetworkFailure().Remove(InternalNotifyNetworkFailureDelegateHandle);
+#else
+		GEngine->OnNetworkFailure().RemoveUObject(this, &UClientUnitTest::InternalNotifyNetworkFailure);
+#endif
 	}
 
 	// Immediately cleanup (or rather, start of next tick, as that's earliest possible time) after sending the RPC
@@ -1306,7 +1334,7 @@ TWeakPtr<FUnitTestProcess> UClientUnitTest::StartUnitTestClient(FString ConnectI
 	{
 		auto CurHandle = ReturnVal.Pin();
 
-		// @todo JohnB: If you add support for multiple clients, make the log prefix numbered, also try to differentiate colours
+		// @todo #JohnBMultiClient: If you add support for multiple clients, make the log prefix numbered, also try to differentiate colours
 		CurHandle->ProcessTag = TEXT("Client");
 		CurHandle->BaseLogType = ELogType::Client;
 		CurHandle->LogPrefix = TEXT("[CLIENT]");
@@ -1365,7 +1393,7 @@ void UClientUnitTest::ShutdownUnitTestProcess(TSharedPtr<FUnitTestProcess> InHan
 		UNIT_STATUS_LOG(ELogType::StatusVerbose, TEXT("%s"), *LogMsg);
 
 
-		// @todo JohnB: Restore 'true' here, once the issue where killing child processes sometimes kills all processes, is fixed
+		// @todo #JohnBHighPri: Restore 'true' here, once the issue where killing child processes sometimes kills all processes, is fixed
 		FPlatformProcess::TerminateProc(InHandle->ProcessHandle);//, true);
 
 #if TARGET_UE4_CL < CL_CLOSEPROC
@@ -1523,7 +1551,8 @@ void UClientUnitTest::CheckOutputForError(TSharedPtr<FUnitTestProcess> InProcess
 				InProcess->ErrorLogStage = EErrorLogStage::ELS_ErrorStart;
 
 				// Reset the timeout for both the unit test and unit test connection here, as callstack logs are prone to failure
-				ResetTimeout(TEXT("Detected crash."), true);
+				//	(do it for at least two minutes as well, as it can take a long time to get a crashlog)
+				ResetTimeout(TEXT("Detected crash."), true, 120);
 			}
 		}
 
@@ -1574,7 +1603,7 @@ void UClientUnitTest::CheckOutputForError(TSharedPtr<FUnitTestProcess> InProcess
 
 
 #if !UE_BUILD_SHIPPING
-// @todo JohnB: 'HookOrigin' is not optimal - try to at least make it a UClientUnitTest pointer
+// @todo #JohnB: 'HookOrigin' is not optimal - try to at least make it a UClientUnitTest pointer
 bool UClientUnitTest::InternalScriptProcessEvent(AActor* Actor, UFunction* Function, void* Parameters, void* HookOrigin)
 {
 	bool bBlockEvent = false;
@@ -1660,6 +1689,12 @@ void UClientUnitTest::UnitTick(float DeltaTime)
 		}
 	}
 
+	// Prevent net connection timeout in developer mode
+	if (bDeveloperMode)
+	{
+		ResetConnTimeout();
+	}
+
 
 	PollProcessOutput();
 
@@ -1711,7 +1746,7 @@ void UClientUnitTest::PollProcessOutput()
 				// Now split up the log into multiple lines
 				TArray<FString> LogLines;
 				
-				// @todo JohnB: Perhaps add support for both platforms line terminator, at some stage
+				// @todo #JohnB: Perhaps add support for both platforms line terminator, at some stage
 #if TARGET_UE4_CL < CL_STRINGPARSEARRAY
 				LogDump.ParseIntoArray(&LogLines, LINE_TERMINATOR, true);
 #else
diff --git a/Engine/Plugins/NetcodeUnitTest/NetcodeUnitTest/Source/NetcodeUnitTest/Private/NUTActor.cpp b/Engine/Plugins/NetcodeUnitTest/NetcodeUnitTest/Source/NetcodeUnitTest/Private/NUTActor.cpp
--- a/Engine/Plugins/NetcodeUnitTest/NetcodeUnitTest/Source/NetcodeUnitTest/Private/NUTActor.cpp
+++ b/Engine/Plugins/NetcodeUnitTest/NetcodeUnitTest/Source/NetcodeUnitTest/Private/NUTActor.cpp
@@ -2,7 +2,7 @@
 
 #include "NetcodeUnitTestPCH.h"
 
-// @todo JohnB: Need to tidy up and fully document this class; not all of the code below is clear
+// @todo #JohnBDoc: Need to tidy up and fully document this class; not all of the code below is clear
 
 #include "NUTActor.h"
 #include "NUTUtilNet.h"
@@ -338,7 +338,7 @@ void ANUTActor::Tick(float DeltaSeconds)
 
 			if (GameState != NULL)
 			{
-				// @todo JohnB: You want this to only happen if no remote players are present (perhaps give all players control instead,
+				// @todo #JohnBReview: You want this to only happen if no remote players are present (perhaps give all players control instead,
 				//				if this becomes a problem - requires setting things up differently though)
 				if (CurNetMode == NM_ListenServer || CurNetMode == NM_Standalone)
 				{
diff --git a/Engine/Plugins/NetcodeUnitTest/NetcodeUnitTest/Source/NetcodeUnitTest/Private/NUTUtil.cpp b/Engine/Plugins/NetcodeUnitTest/NetcodeUnitTest/Source/NetcodeUnitTest/Private/NUTUtil.cpp
--- a/Engine/Plugins/NetcodeUnitTest/NetcodeUnitTest/Source/NetcodeUnitTest/Private/NUTUtil.cpp
+++ b/Engine/Plugins/NetcodeUnitTest/NetcodeUnitTest/Source/NetcodeUnitTest/Private/NUTUtil.cpp
@@ -30,7 +30,7 @@ void NUTUtil::GetUnitTestClassDefList(TArray<UUnitTest*>& OutUnitTestClassDefaul
 
 void NUTUtil::SortUnitTestClassDefList(TArray<UUnitTest*>& InUnitTestClassDefaults)
 {
-	// @todo JohnB: Convert these inline sort functions to lambda's now
+	// @todo #JohnBRefactorLambda: Convert these inline sort functions to lambda's now
 	struct FUnitTestDateSort
 	{
 		FORCEINLINE bool operator()(const UUnitTest& A, const UUnitTest& B) const
@@ -69,7 +69,6 @@ void NUTUtil::SortUnitTestClassDefList(TArray<UUnitTest*>& InUnitTestClassDefaul
 
 
 	// First reorder the unit test classes by date, then grab the unit test types by date, then group them by type/date
-	// @todo JohnB: You can tweak the order here manually, e.g. by manually adding certain categories as hardcoded
 	TArray<FString> ListTypes;
 
 	InUnitTestClassDefaults.Sort(FUnitTestDateSort());
diff --git a/Engine/Plugins/NetcodeUnitTest/NetcodeUnitTest/Source/NetcodeUnitTest/Private/NUTUtilProfiler.cpp b/Engine/Plugins/NetcodeUnitTest/NetcodeUnitTest/Source/NetcodeUnitTest/Private/NUTUtilProfiler.cpp
--- a/Engine/Plugins/NetcodeUnitTest/NetcodeUnitTest/Source/NetcodeUnitTest/Private/NUTUtilProfiler.cpp
+++ b/Engine/Plugins/NetcodeUnitTest/NetcodeUnitTest/Source/NetcodeUnitTest/Private/NUTUtilProfiler.cpp
@@ -23,7 +23,11 @@ void FFrameProfiler::Start()
 
 		const FStatsThreadState& Stats = FStatsThreadState::GetLocalState();
 
+#if TARGET_UE4_CL >= CL_DEPRECATEDEL
 		OnNewFrameDelegateHandle = Stats.NewFrameDelegate.AddRaw(this, &FFrameProfiler::OnNewFrame);
+#else
+		Stats.NewFrameDelegate.AddRaw(this, &FFrameProfiler::OnNewFrame);
+#endif
 	}
 }
 
@@ -33,13 +37,18 @@ void FFrameProfiler::Stop()
 	{
 		const FStatsThreadState& Stats = FStatsThreadState::GetLocalState();
 
+#if TARGET_UE4_CL >= CL_DEPRECATEDEL
 		Stats.NewFrameDelegate.Remove(OnNewFrameDelegateHandle);
+#else
+		Stats.NewFrameDelegate.RemoveRaw(this, &FFrameProfiler::OnNewFrame);
+#endif
+
 
 		StatsMasterEnableSubtract();
 		bActive = false;
 
 
-		// @todo JohnB: For the moment, these objects are not tracked after creation, and are responsible for deleting themselves
+		// @todo #JohnB: For the moment, these objects are not tracked after creation, and are responsible for deleting themselves
 		delete this;
 	}
 }
@@ -82,7 +91,6 @@ void FFrameProfiler::OnNewFrame(int64 Frame)
 		// Since division is done using integers, the value must be pre-multiplied, or the division result is fractional/not-an-integer
 		float FramePercent = (float)((TotalDuration * 10000) / TotalFrameTime) * 0.01f;
 
-		// @todo JohnB: Remove
 #if 0
 		UE_LOG(LogUnitTest, Log, TEXT("Targeted stat found: Duration: %u, FrameTime: %f"),
 				TotalDuration,
diff --git a/Engine/Plugins/NetcodeUnitTest/NetcodeUnitTest/Source/NetcodeUnitTest/Private/Net/NUTUtilNet.cpp b/Engine/Plugins/NetcodeUnitTest/NetcodeUnitTest/Source/NetcodeUnitTest/Private/Net/NUTUtilNet.cpp
--- a/Engine/Plugins/NetcodeUnitTest/NetcodeUnitTest/Source/NetcodeUnitTest/Private/Net/NUTUtilNet.cpp
+++ b/Engine/Plugins/NetcodeUnitTest/NetcodeUnitTest/Source/NetcodeUnitTest/Private/Net/NUTUtilNet.cpp
@@ -365,7 +365,8 @@ bool NUTNet::CreateFakePlayer(UWorld* InWorld, UNetDriver*& InNetDriver, FString
 			FString VersionStr = GEngineVersion.ToString(EVersionComponent::Minor);
 			int32 VersionDelim = VersionStr.Find(TEXT("."));
 			int32 MajorVersion = FCString::Atoi(*VersionStr.Left(VersionDelim));
-			int32 MinorVersion = FCString::Atoi(*VersionStr.Right(VersionDelim));
+			int32 MinorVersion = FCString::Atoi(*VersionStr.Mid(VersionDelim+1));
+
 
 			bool bOldProtocol = (MajorVersion <= 4 && MinorVersion <= 7) &&
 				/** Exception for UT (treat 4.7 as having the new protocol) */
diff --git a/Engine/Plugins/NetcodeUnitTest/NetcodeUnitTest/Source/NetcodeUnitTest/Private/Net/UnitTestChannel.cpp b/Engine/Plugins/NetcodeUnitTest/NetcodeUnitTest/Source/NetcodeUnitTest/Private/Net/UnitTestChannel.cpp
--- a/Engine/Plugins/NetcodeUnitTest/NetcodeUnitTest/Source/NetcodeUnitTest/Private/Net/UnitTestChannel.cpp
+++ b/Engine/Plugins/NetcodeUnitTest/NetcodeUnitTest/Source/NetcodeUnitTest/Private/Net/UnitTestChannel.cpp
@@ -19,7 +19,7 @@ UUnitTestChannel::UUnitTestChannel(const FObjectInitializer& ObjectInitializer)
 }
 
 /**
- * @todo JohnB
+ * @todo #JohnBDoc
  */
 void UUnitTestChannel::Init(UNetConnection* InConnection, int32 InChIndex, bool InOpenedLocally)
 {
@@ -38,27 +38,23 @@ void UUnitTestChannel::Init(UNetConnection* InConnection, int32 InChIndex, bool
 }
 
 /**
- * @todo JohnB
+ * @todo #JohnBDoc
  */
 void UUnitTestChannel::ReceivedBunch(FInBunch& Bunch)
 {
-	// @todo JohnB: Add default behaviour for when not bound, to output channel data someplace (perhaps the log, if not too verbose?)
-
-	// @todo JohnB: Perhaps make the delegate a multicast one, since that may be useful for multiple processing
-
 	// Pass on to delegate
 	ReceivedBunchDel.ExecuteIfBound(Bunch);
 }
 
 /**
- * @todo JohnB
+ * @todo #JohnBDoc
  */
 void UUnitTestChannel::Tick()
 {
 	Super::Tick();
 
 	// Copied from the control channel code
-	// @todo JohnB: This can spam the log sometimes, upon unit test error; fix this
+	// @todo #JohnBBug: This can spam the log sometimes, upon unit test error; fix this
 	if (!OpenAcked)
 	{
 		int32 Count = 0;
diff --git a/Engine/Plugins/NetcodeUnitTest/NetcodeUnitTest/Source/NetcodeUnitTest/Private/Net/UnitTestNetConnection.cpp b/Engine/Plugins/NetcodeUnitTest/NetcodeUnitTest/Source/NetcodeUnitTest/Private/Net/UnitTestNetConnection.cpp
--- a/Engine/Plugins/NetcodeUnitTest/NetcodeUnitTest/Source/NetcodeUnitTest/Private/Net/UnitTestNetConnection.cpp
+++ b/Engine/Plugins/NetcodeUnitTest/NetcodeUnitTest/Source/NetcodeUnitTest/Private/Net/UnitTestNetConnection.cpp
@@ -37,8 +37,13 @@ void UUnitTestNetConnection::InitBase(UNetDriver* InDriver, class FSocket* InSoc
 #endif
 }
 
+#if TARGET_UE4_CL < CL_INITCONNPARAM
 void UUnitTestNetConnection::InitConnection(UNetDriver* InDriver, EConnectionState InState, const FURL& InURL,
-											int32 InConnectionSpeed, int32 InMaxPacket)
+											int32 InConnectionSpeed/*=0*/)
+#else
+void UUnitTestNetConnection::InitConnection(UNetDriver* InDriver, EConnectionState InState, const FURL& InURL,
+											int32 InConnectionSpeed/*=0*/, int32 InMaxPacket/*=0*/)
+#endif
 {
 	Super::InitConnection(InDriver, InState, InURL, InConnectionSpeed);
 
@@ -89,7 +94,7 @@ void UUnitTestNetConnection::HandleClientPlayer(class APlayerController* PC, cla
 	PlayerController = PC;
 	OwningActor = PC;
 
-	// @todo JohnB: This might cause undesirable behaviour, if - for example - HandleDisconnect gets called by
+	// @todo #JohnBReview: This might cause undesirable behaviour, if - for example - HandleDisconnect gets called by
 	//				RPC's, so may want to create a fake localplayer instead
 	PC->Player = GEngine->GetFirstGamePlayer(NUTUtil::GetPrimaryWorld());
 
diff --git a/Engine/Plugins/NetcodeUnitTest/NetcodeUnitTest/Source/NetcodeUnitTest/Private/Net/UnitTestNetDriver.cpp b/Engine/Plugins/NetcodeUnitTest/NetcodeUnitTest/Source/NetcodeUnitTest/Private/Net/UnitTestNetDriver.cpp
--- a/Engine/Plugins/NetcodeUnitTest/NetcodeUnitTest/Source/NetcodeUnitTest/Private/Net/UnitTestNetDriver.cpp
+++ b/Engine/Plugins/NetcodeUnitTest/NetcodeUnitTest/Source/NetcodeUnitTest/Private/Net/UnitTestNetDriver.cpp
@@ -15,7 +15,7 @@ UUnitTestNetDriver::UUnitTestNetDriver(const FObjectInitializer& ObjectInitializ
 }
 
 /**
- * @todo JohnB
+ * @todo #JohnBDoc
  */
 bool UUnitTestNetDriver::InitConnectionClass()
 {
@@ -36,7 +36,7 @@ bool UUnitTestNetDriver::InitConnectionClass()
 }
 
 /**
- * @todo JohnB
+ * @todo #JohnBDoc
  */
 bool UUnitTestNetDriver::InitConnect(FNetworkNotify* InNotify, const FURL& ConnectURL, FString& Error)
 {
diff --git a/Engine/Plugins/NetcodeUnitTest/NetcodeUnitTest/Source/NetcodeUnitTest/Private/Net/UnitTestPackageMap.cpp b/Engine/Plugins/NetcodeUnitTest/NetcodeUnitTest/Source/NetcodeUnitTest/Private/Net/UnitTestPackageMap.cpp
--- a/Engine/Plugins/NetcodeUnitTest/NetcodeUnitTest/Source/NetcodeUnitTest/Private/Net/UnitTestPackageMap.cpp
+++ b/Engine/Plugins/NetcodeUnitTest/NetcodeUnitTest/Source/NetcodeUnitTest/Private/Net/UnitTestPackageMap.cpp
@@ -6,10 +6,6 @@
 #include "Net/UnitTestPackageMap.h"
 
 
-// @todo JohnB: Should probably add a CL-define which wraps this code, since not all codebases will expose PackageMapClient?
-//				(probably better to just error, and force a code change for that, in that codebase - otherwise not all unit tests work)
-
-
 /**
  * Default constructor
  */
diff --git a/Engine/Plugins/NetcodeUnitTest/NetcodeUnitTest/Source/NetcodeUnitTest/Private/NetcodeUnitTest.cpp b/Engine/Plugins/NetcodeUnitTest/NetcodeUnitTest/Source/NetcodeUnitTest/Private/NetcodeUnitTest.cpp
--- a/Engine/Plugins/NetcodeUnitTest/NetcodeUnitTest/Source/NetcodeUnitTest/Private/NetcodeUnitTest.cpp
+++ b/Engine/Plugins/NetcodeUnitTest/NetcodeUnitTest/Source/NetcodeUnitTest/Private/NetcodeUnitTest.cpp
@@ -40,7 +40,10 @@ class FNetcodeUnitTest : public INetcodeUnitTest
 {
 private:
 	static FWorldDelegates::FWorldInitializationEvent::FDelegate OnWorldCreatedDelegate;
+
+#if TARGET_UE4_CL >= CL_DEPRECATEDEL
 	static FDelegateHandle OnWorldCreatedDelegateHandle;
+#endif
 
 public:
 	/**
@@ -55,7 +58,11 @@ public:
 			OnWorldCreatedDelegate = FWorldDelegates::FWorldInitializationEvent::FDelegate::CreateStatic(
 										&FNetcodeUnitTest::OnWorldCreated);
 
+#if TARGET_UE4_CL >= CL_DEPRECATEDEL
 			OnWorldCreatedDelegateHandle = FWorldDelegates::OnPreWorldInitialization.Add(OnWorldCreatedDelegate);
+#else
+			FWorldDelegates::OnPreWorldInitialization.Add(OnWorldCreatedDelegate);
+#endif
 
 			bSetDelegate = true;
 		}
@@ -97,13 +104,19 @@ public:
 		}
 
 		// Now remove it, so it's only called once
+#if TARGET_UE4_CL >= CL_DEPRECATEDEL
 		FWorldDelegates::OnPreWorldInitialization.Remove(OnWorldCreatedDelegateHandle);
+#else
+		FWorldDelegates::OnPreWorldInitialization.Remove(OnWorldCreatedDelegate);
+#endif
 	}
 };
 
 FWorldDelegates::FWorldInitializationEvent::FDelegate FNetcodeUnitTest::OnWorldCreatedDelegate = NULL;
 
+#if TARGET_UE4_CL >= CL_DEPRECATEDEL
 FDelegateHandle FNetcodeUnitTest::OnWorldCreatedDelegateHandle;
+#endif
 
 
 // Essential for getting the .dll to compile, and for the package to be loadable
diff --git a/Engine/Plugins/NetcodeUnitTest/NetcodeUnitTest/Source/NetcodeUnitTest/Private/Tests/AllUnitTestsTest.cpp b/Engine/Plugins/NetcodeUnitTest/NetcodeUnitTest/Source/NetcodeUnitTest/Private/Tests/AllUnitTestsTest.cpp
--- a/Engine/Plugins/NetcodeUnitTest/NetcodeUnitTest/Source/NetcodeUnitTest/Private/Tests/AllUnitTestsTest.cpp
+++ b/Engine/Plugins/NetcodeUnitTest/NetcodeUnitTest/Source/NetcodeUnitTest/Private/Tests/AllUnitTestsTest.cpp
@@ -9,7 +9,7 @@
 #include "NUTUtil.h"
 
 
-// @todo JohnB: Some unit tests may require special automation flags, and you should automatically parse these from the unit test
+// @todo #JohnBAutomation: Some unit tests may require special automation flags, and you should automatically parse these from the unit test
 //				class itself, if possible (and since this runs all unit tests, you should probably iterate all unit tests that will
 //				be run, to check their requirements)
 
@@ -31,7 +31,7 @@ bool FNUTWaitForUnitTests::Update()
 /**
  * FAutomationConsoleCommand
  *
- * @todo JohnB: AutomationTestCommon.h already implements this, but is private
+ * @todo #JohnBRefactor: AutomationTestCommon.h already implements this, but is private
  */
 DEFINE_LATENT_AUTOMATION_COMMAND_ONE_PARAMETER(FAutomationConsoleCommand, FString, Command);
 
diff --git a/Engine/Plugins/NetcodeUnitTest/NetcodeUnitTest/Source/NetcodeUnitTest/Private/UI/LogWindowManager.cpp b/Engine/Plugins/NetcodeUnitTest/NetcodeUnitTest/Source/NetcodeUnitTest/Private/UI/LogWindowManager.cpp
--- a/Engine/Plugins/NetcodeUnitTest/NetcodeUnitTest/Source/NetcodeUnitTest/Private/UI/LogWindowManager.cpp
+++ b/Engine/Plugins/NetcodeUnitTest/NetcodeUnitTest/Source/NetcodeUnitTest/Private/UI/LogWindowManager.cpp
@@ -21,7 +21,11 @@ FLogWindowManager::~FLogWindowManager()
 
 		if (CurEntry->LogWindow.IsValid())
 		{
+#if TARGET_UE4_CL >= CL_DEPRECATEDEL
 			CurEntry->LogWindow->MultiOnWindowClosed.Remove(OnWindowClosedDelegateHandles.FindRef(CurEntry->LogWindow.Get()));
+#else
+			CurEntry->LogWindow->MultiOnWindowClosed.RemoveRaw(this, &FLogWindowManager::OnWindowClosed);
+#endif
 		}
 	}
 
@@ -29,7 +33,11 @@ FLogWindowManager::~FLogWindowManager()
 	{
 		if (OverflowWindows[i].IsValid())
 		{
+#if TARGET_UE4_CL >= CL_DEPRECATEDEL
 			OverflowWindows[i]->MultiOnWindowClosed.Remove(OnWindowClosedDelegateHandles.FindRef(OverflowWindows[i].Get()));
+#else
+			OverflowWindows[i]->MultiOnWindowClosed.RemoveRaw(this, &FLogWindowManager::OnWindowClosed);
+#endif
 		}
 	}
 }
@@ -74,7 +82,7 @@ void FLogWindowManager::Initialize(int InLogWidth, int InLogHeight)
 		}
 
 
-		// @todo JohnB: Remove this debug code
+		// @todo #JohnBDebug: Remove this debug code
 #if 0
 		for (int i = 0; i<GridSpaces.Num(); i++)
 		{
@@ -128,8 +136,12 @@ TSharedPtr<SLogWindow> FLogWindowManager::CreateLogWindow(FString Title, ELogTyp
 
 	if (ReturnVal.IsValid())
 	{
+#if TARGET_UE4_CL >= CL_DEPRECATEDEL
 		OnWindowClosedDelegateHandles.Add(ReturnVal.Get(),
 											ReturnVal->MultiOnWindowClosed.AddRaw(this, &FLogWindowManager::OnWindowClosed));
+#else
+		ReturnVal->MultiOnWindowClosed.AddRaw(this, &FLogWindowManager::OnWindowClosed);
+#endif
 
 		FSlateApplication::Get().AddWindow(ReturnVal.ToSharedRef());
 
diff --git a/Engine/Plugins/NetcodeUnitTest/NetcodeUnitTest/Source/NetcodeUnitTest/Private/UI/SLogWidget.cpp b/Engine/Plugins/NetcodeUnitTest/NetcodeUnitTest/Source/NetcodeUnitTest/Private/UI/SLogWidget.cpp
--- a/Engine/Plugins/NetcodeUnitTest/NetcodeUnitTest/Source/NetcodeUnitTest/Private/UI/SLogWidget.cpp
+++ b/Engine/Plugins/NetcodeUnitTest/NetcodeUnitTest/Source/NetcodeUnitTest/Private/UI/SLogWidget.cpp
@@ -8,19 +8,19 @@
 #include "SMultiSelectTableRow.h"
 
 
-// @todo JohnB: It would be good to have both a 'Search' filter tab, and a 'Ctrl+F Find' button for each existing tab;
+// @todo #JohnBFeatureUI: It would be good to have both a 'Search' filter tab, and a 'Ctrl+F Find' button for each existing tab;
 //				two different search abilities (should be easy to add too)
 
-// @todo JohnB: Could change the 'close' button, to a '+' button, which clones the current tab and enables search filtering
+// @todo #JohnBFeatureUI: Could change the 'close' button, to a '+' button, which clones the current tab and enables search filtering
 
 
-// @todo JohnB: Perhaps have conditional tabs, which only appear if log entries of this type come up?
-//				(such as: Warning, Error, Debug)
+// @todo #JohnBReview: Perhaps have conditional tabs, which only appear if log entries of this type come up?
+//				(such as: Warning, Error, Debug) UPDATE: Is this not already present?
 
-// @todo JohnB: Consider adding a "TabName (LineCount)" to the tab name (maybe just for important tabs, like Error/Warning)
+// @todo #JohnBFeatureUI: Consider adding a "TabName (LineCount)" to the tab name (maybe just for important tabs, like Error/Warning)
 
 
-// @todo JohnB: Change the opacity, of background highlight for selected log entries - really ugly at the moment
+// @todo #JohnBUI: Change the opacity, of background highlight for selected log entries - really ugly at the moment
 
 
 // Enable access to the private SEditableTextBox.EditableText variable, using the GET_PRIVATE macro
@@ -33,7 +33,7 @@ IMPLEMENT_GET_PRIVATE_VAR(SButton, Style, const FButtonStyle*);
 IMPLEMENT_GET_PRIVATE_FUNC_CONST(SDockTab, GetCurrentStyle, const FDockTabStyle&, void, const);
 
 
-// @todo JohnB: Perhaps move widget searching to a NUTSlate.h file, or such?
+// @todo #JohnBRefactorUI: Perhaps move widget searching to a NUTSlate.h file, or such?
 
 /**
  * Delegate used for recursively iterating a widgets child widgets, and testing if they match a search condition
@@ -211,7 +211,7 @@ void SLogWidget::Construct(const FArguments& Args)
 						})
 						[
 							SNew(SImage)
-							// @todo JohnB: The scaled image looks a bit fuzzy, so perhaps don't do that
+							// @todo #JohnBUI: The scaled image looks a bit fuzzy, so perhaps don't do that
 							//				(if you can figure out how to make it clip rather than scale though, do that)
 #if 1
 							.Image(&FCoreStyle::Get().GetWidgetStyle<FSearchBoxStyle>("SearchBox").GlassImage)
@@ -450,7 +450,7 @@ void SLogWidget::Construct(const FArguments& Args)
 						SAssignNew(ConsoleComboBox, SComboBox<TSharedPtr<FString>>)
 						.OptionsSource(&ConsoleContextList)
 						.ToolTipText(FText::FromString(FString(TEXT("Select the context for executing console commands."))))
-						// @todo JohnB: Too big to inline? Separate into its own function?
+						// @todo #JohnBRefactorUI: Too big to inline? Separate into its own function?
 						.OnGenerateWidget_Lambda(
 								[](TSharedPtr<FString> Item)
 								{
@@ -471,12 +471,12 @@ void SLogWidget::Construct(const FArguments& Args)
 									}
 									else if (ItemStr == TEXT("Client"))
 									{
-										// @todo JohnB: Update when implemented
+										// @todo #JohnBFeatureUI: Update when implemented
 										ToolTipStr = TEXT("(Not yet implemented) ") 
 														TEXT("Execute the command on the client associated with this unit test.");
 									}
 
-									// @todo JohnB: Custom context hints?
+									// @todo #JohnB: Custom context hints?
 
 									return	SNew(STextBlock)
 											.Text(FText::FromString(ItemStr))
@@ -500,7 +500,7 @@ void SLogWidget::Construct(const FArguments& Args)
 				/**
 				 * Console command edit box
 				 */
-				 // @todo JohnB: Borrow the auto-complete from SOutputLog's version of log windows
+				 // @todo #JohnBFeatureUI: Borrow the auto-complete from SOutputLog's version of log windows
 				+SHorizontalBox::Slot()
 					.HAlign(HAlign_Fill)
 					.VAlign(VAlign_Center)
@@ -510,7 +510,7 @@ void SLogWidget::Construct(const FArguments& Args)
 						.HintText(FText::FromString(TEXT("Console")))
 						.ToolTipText(FText::FromString(TEXT("Executes a console command within the specified context.")))
 						.ClearKeyboardFocusOnCommit(false)
-						// @todo JohnB: Too big to inline? Separate into its own function?
+						// @todo #JohnBRefactorUI: Too big to inline? Separate into its own function?
 						.OnTextCommitted_Lambda(
 							[&](const FText& InText, ETextCommit::Type InCommitType)
 							{
@@ -658,7 +658,7 @@ TSharedRef<SDockTab> SLogWidget::SpawnLogTab(const FSpawnTabArgs& InSpawnTabArgs
 		});
 
 
-	// @todo JohnB: Code duplication - this is defined above too - move to an inline function
+	// @todo #JohnBRefactorUI: Code duplication - this is defined above too - move to an inline function
 	auto ArrayAddNew =
 		[] (TArray<TSharedPtr<SWidget>>& InArray) -> TSharedPtr<SWidget>&
 		{
@@ -706,7 +706,7 @@ TSharedRef<SDockTab> SLogWidget::SpawnLogTab(const FSpawnTabArgs& InSpawnTabArgs
 						 */
 						SAssignNew(CurTabInfo->LogListView, SListView<TSharedRef<FLogLine>>)
 						.ListItemsSource(&CurTabInfo->TabLogLines)
-						// @todo JohnB: Probably large enough to separate to its own function now
+						// @todo #JohnBRefactorUI: Probably large enough to separate to its own function now
 						.OnGenerateRow_Lambda(
 							[](TSharedRef<FLogLine> Item, const TSharedRef<STableViewBase>& OwnerTable)
 							{
@@ -1208,7 +1208,7 @@ void SLogWidget::ScrollToText(TSharedRef<FLogTabInfo> InTab, FString FindText, b
 		int32 SearchDir = (bSearchUp ? -1 : 1);
 		int32 FoundIdx = INDEX_NONE;
 
-		for (int32 i=FindStartIdx + SearchDir; i != FindStartIdx; i += SearchDir)
+		for (int32 i=FindStartIdx + SearchDir; true; i += SearchDir)
 		{
 			// When the start/end of the array is reached, wrap-around to the other end
 			if (i < 0 || i >= CurTabLogLines.Num())
@@ -1216,6 +1216,12 @@ void SLogWidget::ScrollToText(TSharedRef<FLogTabInfo> InTab, FString FindText, b
 				i = (bSearchUp ? CurTabLogLines.Num()-1 : 0);
 			}
 
+			// Moved out of 'for' condition, and into the loop, to avoid infinite recursion in rare circumstances
+			if (i == FindStartIdx)
+			{
+				break;
+			}
+
 			if (CurTabLogLines[i]->LogLine->Contains(FindText))
 			{
 				FoundIdx = i;
@@ -1230,7 +1236,7 @@ void SLogWidget::ScrollToText(TSharedRef<FLogTabInfo> InTab, FString FindText, b
 			CurLogListView->RequestListRefresh();
 		}
 
-		// @todo JohnB: Find some way of indicating a failed search
+		// @todo #JohnBUI: Find some way of indicating a failed search
 
 		InTab->bLastFindWasUp = bSearchUp;
 	}
diff --git a/Engine/Plugins/NetcodeUnitTest/NetcodeUnitTest/Source/NetcodeUnitTest/Private/UnitTestCommandlet.cpp b/Engine/Plugins/NetcodeUnitTest/NetcodeUnitTest/Source/NetcodeUnitTest/Private/UnitTestCommandlet.cpp
--- a/Engine/Plugins/NetcodeUnitTest/NetcodeUnitTest/Source/NetcodeUnitTest/Private/UnitTestCommandlet.cpp
+++ b/Engine/Plugins/NetcodeUnitTest/NetcodeUnitTest/Source/NetcodeUnitTest/Private/UnitTestCommandlet.cpp
@@ -10,7 +10,7 @@
 #include "SLogDialog.h"
 
 
-// @todo JohnB: Fix StandaloneRenderer support for static builds
+// @todo #JohnBLowPri: Fix StandaloneRenderer support for static builds
 #if !IS_MONOLITHIC
 #include "StandaloneRenderer.h"
 #endif
@@ -21,7 +21,7 @@
 #include "SlateBasics.h"
 
 
-// @todo JohnB: If you later end up doing test client instances that are unit tested against client netcode
+// @todo #JohnB_NUTClient: If you later end up doing test client instances that are unit tested against client netcode
 //				(same way as you do server instances at the moment, for testing server netcode),
 //				then this commandlet code is probably an excellent base for setting up minimal standalone clients like that,
 //				as it's intended to strip-down running unit tests, from a minimal client itself
@@ -55,18 +55,18 @@ UUnitTestCommandlet::UUnitTestCommandlet(const FObjectInitializer& ObjectInitial
 
 int32 UUnitTestCommandlet::Main(const FString& Params)
 {
-	// @todo JohnB: Fix StandaloneRenderer support for static builds
+	// @todo #JohnBLowPri: Fix StandaloneRenderer support for static builds
 #if IS_MONOLITHIC
 	UE_LOG(LogUnitTest, Log, TEXT("NetcodeUnitTest commandlet not currently supported in static/monolithic builds"));
 
 #else
 	GIsRequestingExit = false;
 
-	// @todo JohnB: Steam detection doesn't seem to work this early on, but does work further down the line;
+	// @todo #JohnBLowPri: Steam detection doesn't seem to work this early on, but does work further down the line;
 	//				try to find a way, to detect it as early as possible
 
 	// NetcodeUnitTest is not compatible with Steam; if Steam is running/detected, abort immediately
-	// @todo JohnB: Add support for Steam
+	// @todo #JohnBLowPri: Add support for Steam
 	if (NUTNet::IsSteamNetDriverAvailable())
 	{
 		UE_LOG(LogUnitTest, Log, TEXT("NetcodeUnitTest does not currently support Steam. Close Steam before running."));
@@ -209,7 +209,7 @@ int32 UUnitTestCommandlet::Main(const FString& Params)
 
 void UUnitTestCommandlet::CreateCustomEngine(const FString& Params)
 {
-	// @todo JohnB: Fix StandaloneRenderer support for static builds
+	// @todo #JohnBLowPri: Fix StandaloneRenderer support for static builds
 #if !IS_MONOLITHIC
 	FSlateApplication::InitializeAsStandaloneApplication(GetStandardStandaloneRenderer());
 #endif
diff --git a/Engine/Plugins/NetcodeUnitTest/NetcodeUnitTest/Source/NetcodeUnitTest/Private/UnitTestEnvironment.cpp b/Engine/Plugins/NetcodeUnitTest/NetcodeUnitTest/Source/NetcodeUnitTest/Private/UnitTestEnvironment.cpp
--- a/Engine/Plugins/NetcodeUnitTest/NetcodeUnitTest/Source/NetcodeUnitTest/Private/UnitTestEnvironment.cpp
+++ b/Engine/Plugins/NetcodeUnitTest/NetcodeUnitTest/Source/NetcodeUnitTest/Private/UnitTestEnvironment.cpp
@@ -97,7 +97,16 @@ FString FUnitTestEnvironment::GetDefaultClientParameters()
 	FString ReturnVal = TEXT("");
 	FString CmdLineClientParms;
 
-	ReturnVal = TEXT("-nullrhi -windowed -resx=640 -resy=480");
+	bool bDebugClient = FParse::Param(FCommandLine::Get(), TEXT("UnitTestClientDebug"));
+
+	if (!bDebugClient)
+	{
+		ReturnVal = TEXT("-nullrhi -windowed -resx=640 -resy=480");
+	}
+	else
+	{
+		ReturnVal = TEXT("-windowed -resx=1024 -resy=768");
+	}
 
 	// @todo JohnB: Add merging of LogCmds, from above function, if adding default LogCmds for any game
 
diff --git a/Engine/Plugins/NetcodeUnitTest/NetcodeUnitTest/Source/NetcodeUnitTest/Private/UnitTestManager.cpp b/Engine/Plugins/NetcodeUnitTest/NetcodeUnitTest/Source/NetcodeUnitTest/Private/UnitTestManager.cpp
--- a/Engine/Plugins/NetcodeUnitTest/NetcodeUnitTest/Source/NetcodeUnitTest/Private/UnitTestManager.cpp
+++ b/Engine/Plugins/NetcodeUnitTest/NetcodeUnitTest/Source/NetcodeUnitTest/Private/UnitTestManager.cpp
@@ -17,7 +17,7 @@
 #include "NUTUtilNet.h"
 #include "NUTUtilProfiler.h"
 
-// @todo JohnB: Add an overall-timer, and then start debugging the memory management in more detail
+// @todo #JohnBFeature: Add an overall-timer, and then start debugging the memory management in more detail
 
 
 UUnitTestManager* GUnitTestManager = NULL;
@@ -34,10 +34,10 @@ static TMap<FString, FString> UnsupportedUnitTests;
 
 
 
-// @todo JohnB: Add detection for unit tests that either 1: Never end (e.g. if they become outdated and break, or 2: keep on hitting
+// @todo #JohnBHighPri: Add detection for unit tests that either 1: Never end (e.g. if they become outdated and break, or 2: keep on hitting
 //				memory limits and get auto-closed and re-run recursively, forever
 
-// @todo JohnB: For unit tests that are running in the background (no log window for controlling them),
+// @todo #JohnBFeatureUI: For unit tests that are running in the background (no log window for controlling them),
 //				add a drop-down listbox button to the status window, for re-opening background unit test log windows
 //				(can even add an option, to have them all running in background by default, so you have to do this to open
 //				log window at all)
@@ -54,6 +54,7 @@ UUnitTestManager::UUnitTestManager(const FObjectInitializer& ObjectInitializer)
 	, bCapUnitTestMemory(false)
 	, MaxMemoryPercent(0)
 	, AutoCloseMemoryPercent(0)
+	, MaxAutoCloseCount(0)
 	, PendingUnitTests()
 	, ActiveUnitTests()
 	, FinishedUnitTests()
@@ -65,6 +66,7 @@ UUnitTestManager::UUnitTestManager(const FObjectInitializer& ObjectInitializer)
 	, DialogWindows()
 	, StatusWindow()
 	, AbortAllDialog()
+	, LastMemoryLimitHit(0.0)
 	, MemoryTickCountdown(0)
 	, MemoryUsageUponCountdown(0)
 {
@@ -100,6 +102,8 @@ void UUnitTestManager::Initialize()
 		// Since the above can undershoot, the limit at which unit tests are automatically terminated, is a bit higher
 		AutoCloseMemoryPercent = 90;
 
+		MaxAutoCloseCount = 4;
+
 		UE_LOG(LogUnitTest, Log, TEXT("Creating initial unit test config file"));
 
 		SaveConfig();
@@ -223,7 +227,7 @@ bool UUnitTestManager::QueueUnitTest(UClass* UnitTestClass, bool bRequeued/*=fal
 	{
 		ELogType StatusType = ELogType::StyleBold;
 
-		// @todo JohnB: This should be enabled by default instead, so automation testing does give an error when a game doesn't
+		// @todo #JohnBAutomation: This should be enabled by default instead, so automation testing does give an error when a game doesn't
 		//				have a unit test environment. You should probably setup a whitelist of 'known-unsupported-games' somewhere,
 		//				to keep track of what games need support added, without breaking automation testing (but without breaking the
 		//				flow of setting-up/running automation tests)
@@ -242,84 +246,88 @@ bool UUnitTestManager::QueueUnitTest(UClass* UnitTestClass, bool bRequeued/*=fal
 
 void UUnitTestManager::PollUnitTestQueue()
 {
-	// @todo JohnB: Maybe consider staggering the start of unit tests, as perhaps this may give late-starters
+	// @todo #JohnB: Maybe consider staggering the start of unit tests, as perhaps this may give late-starters
 	//				a boost from OS precaching of file data, and may also help stagger the 'peak memory' time,
 	//				which often gets multiple unit tests closed early due to happening at the same time?
 
-	// Keep kicking off unit tests in order, until the list is empty, or until the unit test cap is reached
-	for (int32 i=0; i<PendingUnitTests.Num(); i++)
+	// If the memory limit was recently hit, wait a number of seconds before launching any more unit tests
+	if (PendingUnitTests.Num() > 0 && (FPlatformTime::Seconds() - LastMemoryLimitHit) > 4.0)
 	{
-		bool bAlreadyRemoved = false;
+		// Keep kicking off unit tests in order, until the list is empty, or until the unit test cap is reached
+		for (int32 i=0; i<PendingUnitTests.Num(); i++)
+		{
+			bool bAlreadyRemoved = false;
 
-		// Lambda for remove-safe and multi-remove-safe handling within this loop
-		auto RemoveCurrent = [&]()
-			{
-				if (!bAlreadyRemoved)
+			// Lambda for remove-safe and multi-remove-safe handling within this loop
+			auto RemoveCurrent = [&]()
 				{
-					bAlreadyRemoved = true;
-					PendingUnitTests.RemoveAt(i);
-					i--;
-				}
-			};
-
+					if (!bAlreadyRemoved)
+					{
+						bAlreadyRemoved = true;
+						PendingUnitTests.RemoveAt(i);
+						i--;
+					}
+				};
 
-		UClass* CurUnitTestClass = PendingUnitTests[i];
-		bool bWithinUnitTestLimits = ActiveUnitTests.Num() == 0 || WithinUnitTestLimits(CurUnitTestClass);
 
-		// This unit test isn't within limits, continue to the next one and see if it fits
-		if (!bWithinUnitTestLimits)
-		{
-			continue;
-		}
+			UClass* CurUnitTestClass = PendingUnitTests[i];
+			bool bWithinUnitTestLimits = ActiveUnitTests.Num() == 0 || WithinUnitTestLimits(CurUnitTestClass);
 
+			// This unit test isn't within limits, continue to the next one and see if it fits
+			if (!bWithinUnitTestLimits)
+			{
+				continue;
+			}
 
-		UUnitTest* CurUnitTestDefault = Cast<UUnitTest>(CurUnitTestClass->GetDefaultObject());
 
-		if (CurUnitTestDefault != NULL)
-		{
-			auto CurUnitTest = NewObject<UUnitTest>(GetTransientPackage(), CurUnitTestClass);
+			UUnitTest* CurUnitTestDefault = Cast<UUnitTest>(CurUnitTestClass->GetDefaultObject());
 
-			if (CurUnitTest != NULL)
+			if (CurUnitTestDefault != NULL)
 			{
-				if (UUnitTest::UnitEnv != NULL)
+				auto CurUnitTest = NewObject<UUnitTest>(GetTransientPackage(), CurUnitTestClass);
+
+				if (CurUnitTest != NULL)
 				{
-					CurUnitTest->InitializeEnvironmentSettings();
-				}
+					if (UUnitTest::UnitEnv != NULL)
+					{
+						CurUnitTest->InitializeEnvironmentSettings();
+					}
 
-				// Remove from PendingUnitTests, and add to ActiveUnitTests
-				RemoveCurrent();
-				ActiveUnitTests.Add(CurUnitTest);
+					// Remove from PendingUnitTests, and add to ActiveUnitTests
+					RemoveCurrent();
+					ActiveUnitTests.Add(CurUnitTest);
 
-				// Create the log window (if starting the unit test fails, this is unset during cleanup)
-				if (!FApp::IsUnattended())
-				{
-					OpenUnitTestLogWindow(CurUnitTest);
-				}
+					// Create the log window (if starting the unit test fails, this is unset during cleanup)
+					if (!FApp::IsUnattended())
+					{
+						OpenUnitTestLogWindow(CurUnitTest);
+					}
 
 
-				if (CurUnitTest->StartUnitTest())
-				{
-					STATUS_LOG(ELogType::StatusImportant, TEXT("Started unit test '%s'"), *CurUnitTest->GetUnitTestName());
+					if (CurUnitTest->StartUnitTest())
+					{
+						STATUS_LOG(ELogType::StatusImportant, TEXT("Started unit test '%s'"), *CurUnitTest->GetUnitTestName());
+					}
+					else
+					{
+						STATUS_LOG(ELogType::StatusError | ELogType::StyleBold, TEXT("Failed to kickoff unit test '%s'"),
+										*CurUnitTestDefault->GetUnitTestName());
+					}
 				}
 				else
 				{
-					STATUS_LOG(ELogType::StatusError | ELogType::StyleBold, TEXT("Failed to kickoff unit test '%s'"),
+					STATUS_LOG(ELogType::StatusError | ELogType::StyleBold, TEXT("Failed to construct unit test: %s"),
 									*CurUnitTestDefault->GetUnitTestName());
 				}
 			}
 			else
 			{
-				STATUS_LOG(ELogType::StatusError | ELogType::StyleBold, TEXT("Failed to construct unit test: %s"),
-								*CurUnitTestDefault->GetUnitTestName());
+				STATUS_LOG(ELogType::StatusError | ELogType::StyleBold, TEXT("Failed to find default object for unit test class '%s'"),
+								*CurUnitTestClass->GetName());
 			}
-		}
-		else
-		{
-			STATUS_LOG(ELogType::StatusError | ELogType::StyleBold, TEXT("Failed to find default object for unit test class '%s'"),
-							*CurUnitTestClass->GetName());
-		}
 
-		RemoveCurrent();
+			RemoveCurrent();
+		}
 	}
 }
 
@@ -327,7 +335,7 @@ bool UUnitTestManager::WithinUnitTestLimits(UClass* PendingUnitTest/*=NULL*/)
 {
 	bool bReturnVal = false;
 
-	// @todo JohnB: Could do with non-spammy logging of when limits are reached
+	// @todo #JohnBReview: Could do with non-spammy logging of when limits are reached
 
 	// Check max unit test count
 	bReturnVal = !bCapUnitTestCount || ActiveUnitTests.Num() < MaxUnitTestCount;
@@ -345,7 +353,7 @@ bool UUnitTestManager::WithinUnitTestLimits(UClass* PendingUnitTest/*=NULL*/)
 
 	if (bReturnVal && !bCapUnitTestCount && ActiveUnitTests.Num() >= FirstRunCap)
 	{
-		// @todo JohnB: Add prominent logging for hitting this cap - preferably 'status' log window
+		// @todo #JohnB: Add prominent logging for hitting this cap - preferably 'status' log window
 
 		uint32 FirstRunCount = 0;
 
@@ -484,14 +492,10 @@ bool UUnitTestManager::WithinUnitTestLimits(UClass* PendingUnitTest/*=NULL*/)
 		SIZE_T EstimatedPeakPhysicalMem = (UsedPhysicalMem - CurrentTotalUnitMemUsage) + WorstCaseTotalUnitMemUsage;
 
 		bReturnVal = MaxPhysicalMem > EstimatedPeakPhysicalMem;
-
-
-		// @todo JohnB: Remove this debug code
-		//UE_LOG(LogUnitTest, Log, TEXT("Estimated peak physical mem: %llumb"), (EstimatedPeakPhysicalMem / (SIZE_T)1048576));
 	}
 
 
-	// @todo JohnB: How to improve unit test memory estimates:
+	// @todo #JohnBFeature: How to improve unit test memory estimates:
 	// NOTE: Not to be done, unless above implementation proves problematic (seems to provide decent estimates so-far)
 	//	- Can improve future-memory-usage estimation, by storing time-series stats, mapping the memory usage of each unit test
 	//		- Problem here though, is this is liable to vary over time, so perhaps only do this for each unit tests 'most recent run'
@@ -614,8 +618,6 @@ void UUnitTestManager::NotifyLogWindowClosed(const TSharedRef<SWindow>& ClosedWi
 			if (!CurUnitTest->bCompleted && !CurUnitTest->bAborted)
 			{
 				// Show a message box, asking the player if they'd like to also abort the unit test
-				// @todo JohnB: Perhaps add an option somewhere, to abort-by-default as well, to save the time spent on UI
-				//				(could even make it a tickbox on the status window)
 				FString UnitTestName = CurUnitTest->GetUnitTestName();
 
 				FText CloseMsg = FText::FromString(FString::Printf(TEXT("Abort unit test '%s'? (currently running in background)"),
@@ -698,8 +700,11 @@ void UUnitTestManager::DumpStatus(bool bForce/*=false*/)
 		// doesn't disrupt the flow of text (otherwise, all the other events/updates in the status window, become hard to read)
 		STATUS_SET_COLOR(FLinearColor(0.f, 1.f, 1.f));
 
-		// @todo JohnB: Prettify the status command, by seeing if you can evenly space unit test info into columns;
+		// @todo #JohnB: Prettify the status command, by seeing if you can evenly space unit test info into columns;
 		//				use the console monospacing to help do this (if it looks ok)
+
+		// @todo #JohnBFeatureUI: Better yet, replace the unit test status, text, with an actual list of running unit tests
+
 		SIZE_T TotalMemoryUsage = 0;
 
 		STATUS_LOG(, TEXT(""));
@@ -735,7 +740,7 @@ void UUnitTestManager::DumpStatus(bool bForce/*=false*/)
 	}
 }
 
-void UUnitTestManager::PrintUnitTestResult(UUnitTest* InUnitTest, bool bFinalSummary/*=false*/)
+void UUnitTestManager::PrintUnitTestResult(UUnitTest* InUnitTest, bool bFinalSummary/*=false*/, bool bUnfinished/*=false*/)
 {
 	static const UEnum* VerificationStateEnum = FindObject<UEnum>(ANY_PACKAGE, TEXT("EUnitTestVerification"));
 
@@ -750,7 +755,7 @@ void UUnitTestManager::PrintUnitTestResult(UUnitTest* InUnitTest, bool bFinalSum
 	}
 
 	STATUS_LOG_OBJ(InUnitTest, ELogType::StatusImportant,	TEXT("  - Result: %s"),
-					*VerificationStateEnum->GetEnumName((uint32)UnitTestResult));
+					(bUnfinished ? TEXT("Aborted/Unfinished") : *VerificationStateEnum->GetEnumName((uint32)UnitTestResult)));
 	STATUS_LOG_OBJ(InUnitTest, ELogType::StatusVerbose,	TEXT("  - Execution Time: %f"), InUnitTest->LastExecutionTime);
 
 
@@ -807,8 +812,12 @@ void UUnitTestManager::PrintUnitTestResult(UUnitTest* InUnitTest, bool bFinalSum
 
 	EUnitTestVerification ExpectedResult = InUnitTest->GetExpectedResult();
 
-	// @todo JohnB: Perhaps make this an execution-blocking error instead? (simplifies the code/error-cases here)
-	if (ExpectedResult == EUnitTestVerification::Unverified)
+	if (bUnfinished)
+	{
+		STATUS_LOG_OBJ(InUnitTest, ELogType::StatusWarning | StatusAutomationFlag | ELogType::StyleBold,
+						TEXT("  - WARNING: Unit test was aborted and could not be successfully run."))
+	}
+	else if (ExpectedResult == EUnitTestVerification::Unverified)
 	{
 		STATUS_LOG_OBJ(InUnitTest, ELogType::StatusError | StatusAutomationFlag | ELogType::StyleBold,
 						TEXT("  - Unit test does not have 'ExpectedResult' set"));
@@ -874,7 +883,7 @@ void UUnitTestManager::PrintUnitTestResult(UUnitTest* InUnitTest, bool bFinalSum
 
 void UUnitTestManager::PrintFinalSummary()
 {
-	// @todo JohnB: Add some extra stats eventually, such as overall time taken etc.
+	// @todo #JohnBFeatureUI: Add some extra stats eventually, such as overall time taken etc.
 
 	STATUS_LOG(ELogType::StatusImportant, TEXT(""));
 	STATUS_LOG(ELogType::StatusImportant, TEXT(""));
@@ -899,8 +908,9 @@ void UUnitTestManager::PrintFinalSummary()
 	UnsupportedUnitTests.Empty();
 
 
-	// Then print the aborted unit tests
+	// Then print the aborted unit tests, and unit tests that have aborted so many times that they can't complete
 	TArray<FString> AbortList;
+	TArray<UUnitTest*> UnfinishedUnitTests;
 
 	for (auto CurUnitTest : FinishedUnitTests)
 	{
@@ -922,69 +932,120 @@ void UUnitTestManager::PrintFinalSummary()
 			{
 				NumberOfAborts++;
 				AbortList.RemoveAt(DupeIdx);
-				DupeIdx--;
 			}
+
+			// NOTE: DupeIdx is invalid past here (as it can not be decremented after RemoveAt above)
+		}
+
+
+		// If the unit test did not have a successful execution, note this
+		UUnitTest* UnfinishedTest = NULL;
+
+		bool bUnitTestCompleted = false;
+
+		for (auto CurUnitTest : FinishedUnitTests)
+		{
+			if (CurUnitTest->GetUnitTestName() == CurAbort)
+			{
+				UnfinishedTest = CurUnitTest;
+
+				if (!CurUnitTest->bAborted)
+				{
+					bUnitTestCompleted = true;
+					break;
+				}
+			}
+		}
+
+		if (!bUnitTestCompleted)
+		{
+			UnfinishedUnitTests.Add(UnfinishedTest);
 		}
 
+
+		FString AbortMsg;
+
 		if (NumberOfAborts == 1)
 		{
-			STATUS_LOG(ELogType::StatusWarning | ELogType::StyleBold, TEXT("%s: Aborted."), *CurAbort);
+			AbortMsg = FString::Printf(TEXT("%s: Aborted."), *CurAbort);
+		}
+		else
+		{
+			AbortMsg = FString::Printf(TEXT("%s: Aborted ('%i' times)."), *CurAbort, NumberOfAborts);
+		}
+
+		if (bUnitTestCompleted)
+		{
+			STATUS_LOG(ELogType::StatusWarning | ELogType::StyleBold, TEXT("%s"), *AbortMsg);
 		}
 		else
 		{
-			STATUS_LOG(ELogType::StatusWarning | ELogType::StyleBold, TEXT("%s: Aborted ('%i' times)."), *CurAbort, NumberOfAborts);
+			AbortMsg += TEXT(" Failed to successfully retry unit test after aborting.");
+
+			STATUS_LOG(ELogType::StatusWarning | ELogType::StyleBold | ELogType::StyleItalic, TEXT("%s"), *AbortMsg);
 		}
 	}
 
 
-	if (AbortList.Num() > 0 || UnsupportedUnitTests.Num() > 0)
+	if (AbortList.Num() > 0 || UnfinishedUnitTests.Num() > 0)
 	{
 		STATUS_LOG(ELogType::StatusImportant, TEXT(""));
 		STATUS_LOG(ELogType::StatusImportant, TEXT(""));
 	}
 
-	// Now print the completed unit tests, which have more detailed information
-	for (auto CurUnitTest : FinishedUnitTests)
-	{
-		if (!CurUnitTest->bAborted)
+	// Now print the completed and unfinished unit tests, which have more detailed information
+	auto StatusPrintResult =
+		[&](UUnitTest* CurUnitTest, bool bUnfinished)
 		{
-			STATUS_SET_COLOR(FLinearColor(0.25f, 0.25f, 0.25f));
+			if (!CurUnitTest->bAborted || bUnfinished)
+			{
+				STATUS_SET_COLOR(FLinearColor(0.25f, 0.25f, 0.25f));
 
-			STATUS_LOG(ELogType::StatusImportant | ELogType::StatusAutomation | ELogType::StyleBold | ELogType::StyleUnderline,
-						TEXT("%s:"), *CurUnitTest->GetUnitTestName());
+				STATUS_LOG(ELogType::StatusImportant | ELogType::StatusAutomation | ELogType::StyleBold | ELogType::StyleUnderline,
+							TEXT("%s:"), *CurUnitTest->GetUnitTestName());
 
-			STATUS_RESET_COLOR();
+				STATUS_RESET_COLOR();
 
 
-			// Print out the main result header
-			PrintUnitTestResult(CurUnitTest, true);
+				// Print out the main result header
+				PrintUnitTestResult(CurUnitTest, true, bUnfinished);
 
 
-			// Now print out the full event history
-			bool bHistoryContainsImportant = CurUnitTest->StatusLogSummary.ContainsByPredicate(
-				[](const TSharedPtr<FUnitStatusLog>& CurEntry)
+				// Now print out the full event history
+				bool bHistoryContainsImportant = CurUnitTest->StatusLogSummary.ContainsByPredicate(
+					[](const TSharedPtr<FUnitStatusLog>& CurEntry)
+					{
+						return ((CurEntry->LogType & ELogType::StatusImportant) == ELogType::StatusImportant);
+					});
+
+
+				if (bHistoryContainsImportant)
 				{
-					return ((CurEntry->LogType & ELogType::StatusImportant) == ELogType::StatusImportant);
-				});
+					STATUS_LOG(ELogType::StatusImportant, TEXT("  - Log summary:"));
+				}
+				else
+				{
+					STATUS_LOG(ELogType::StatusVerbose, TEXT("  - Log summary:"));
+				}
 
+				for (auto CurStatusLog : CurUnitTest->StatusLogSummary)
+				{
+					STATUS_LOG(CurStatusLog->LogType, TEXT("      %s"), *CurStatusLog->LogLine);
+				}
 
-			if (bHistoryContainsImportant)
-			{
-				STATUS_LOG(ELogType::StatusImportant, TEXT("  - Log summary:"));
-			}
-			else
-			{
-				STATUS_LOG(ELogType::StatusVerbose, TEXT("  - Log summary:"));
-			}
 
-			for (auto CurStatusLog : CurUnitTest->StatusLogSummary)
-			{
-				STATUS_LOG(CurStatusLog->LogType, TEXT("      %s"), *CurStatusLog->LogLine);
+				STATUS_LOG(ELogType::StatusImportant, TEXT(""));
 			}
+		};
 
+	for (auto CurUnitTest : FinishedUnitTests)
+	{
+		StatusPrintResult(CurUnitTest, false);
+	}
 
-			STATUS_LOG(ELogType::StatusImportant, TEXT(""));
-		}
+	for (auto CurUnitTest : UnfinishedUnitTests)
+	{
+		StatusPrintResult(CurUnitTest, true);
 	}
 }
 
@@ -1034,7 +1095,7 @@ void UUnitTestManager::OpenStatusWindow()
 				{
 					bool bSuccess = false;
 
-					// @todo JohnB: Revisit this - doesn't STATUS_LOG on its own (>not< STATUS_LOG_BASE),
+					// @todo #JohnBReview: Revisit this - doesn't STATUS_LOG on its own (>not< STATUS_LOG_BASE),
 					//				already output to log? Don't think the extra code here is needed.
 
 					// Need an output device redirector, to send console command log output to both GLog and unit test status log,
@@ -1156,6 +1217,8 @@ void UUnitTestManager::Tick(float DeltaTime)
 		MemoryTickCountdown = 10;
 		MemoryUsageUponCountdown = UsedPhysicalMem;
 
+		LastMemoryLimitHit = FPlatformTime::Seconds();
+
 
 		for (int32 i=ActiveUnitTests.Num()-1; i>=0; i--)
 		{
@@ -1169,7 +1232,39 @@ void UUnitTestManager::Tick(float DeltaTime)
 			CurUnitTest->AbortUnitTest();
 			CurUnitTest = NULL;
 
-			QueueUnitTest(UnitTestClass, true);
+
+			if (bAllowRequeuingUnitTests)
+			{
+				bool bAllowRequeue = true;
+
+				// If the number of auto-abort re-queue's is limited, make sure we're not exceeding the limit
+				if (MaxAutoCloseCount > 0)
+				{
+					uint8 CloseCount = 0;
+
+					for (auto CurFinished : FinishedUnitTests)
+					{
+						if (CurFinished->bAborted && CurFinished->GetClass() == UnitTestClass)
+						{
+							CloseCount++;
+						}
+					}
+
+					if (CloseCount >= MaxAutoCloseCount)
+					{
+						STATUS_LOG(ELogType::StatusWarning | ELogType::StyleBold,
+									TEXT("Unit Test '%s' was aborted more than the maximum of '%i' times, and can't be re-queued."),
+									*GetDefault<UUnitTest>(UnitTestClass)->GetUnitTestName(), MaxAutoCloseCount);
+
+						bAllowRequeue = false;
+					}
+				}
+
+				if (bAllowRequeue)
+				{
+					QueueUnitTest(UnitTestClass, true);
+				}
+			}
 
 
 			// Keep closing unit tests, until we get back within memory limits
@@ -1217,14 +1312,14 @@ bool UUnitTestManager::Exec(UWorld* InWorld, const TCHAR* Cmd, FOutputDevice& Ar
 {
 	bool bReturnVal = true;
 
-	// @todo JohnB: Detecting the originating unit test, by World, no longer works due to delayed fake client launches;
+	// @todo #JohnBBug: Detecting the originating unit test, by World, no longer works due to delayed fake client launches;
 	//				either find another way of determining the originating unit test, or give all unit tests a unique World early on,
 	//				before the fake client launch
 
 	FString UnitTestName = FParse::Token(Cmd, false);
 	bool bValidUnitTestName = false;
 
-	// @todo JohnB: Refactor some of this function, so that the unit test defaults scanning/updating,
+	// @todo #JohnBReview: Refactor some of this function, so that the unit test defaults scanning/updating,
 	//				is done in a separate specific function
 	//				IMPORTANT: It should be called multiple times, when needed, not just once,
 	//				as there will likely be unit tests in multiple loaded packages in the future
@@ -1235,7 +1330,7 @@ bool UUnitTestManager::Exec(UWorld* InWorld, const TCHAR* Cmd, FOutputDevice& Ar
 	NUTUtil::GetUnitTestClassDefList(UnitTestClassDefaults);
 
 	// All unit tests should be given a proper date, so give big errors when this is not set
-	// @todo JohnB: Move to the unit test validation function
+	// @todo #JohnBReview: Move to the unit test validation function
 	for (int i=0; i<UnitTestClassDefaults.Num(); i++)
 	{
 		if (UnitTestClassDefaults[i]->GetUnitTestDate() == FDateTime::MinValue())
@@ -1267,7 +1362,7 @@ bool UUnitTestManager::Exec(UWorld* InWorld, const TCHAR* Cmd, FOutputDevice& Ar
 				FFrameProfiler* NewProfiler = new FFrameProfiler(FName(*TargetEvent), FramePercentThreshold);
 				NewProfiler->Start();
 
-				// @todo JohnB: Perhaps add tracking at some stage, and remove self-destruct code from above class
+				// @todo #JohnBLowPri: Perhaps add tracking at some stage, and remove self-destruct code from above class
 			}
 			else
 			{
@@ -1723,7 +1818,7 @@ static bool UnitTestExec(UWorld* InWorld, const TCHAR* Cmd, FOutputDevice& Ar)
 				GTraceManager.Dump(TraceName);
 			}
 			// If no subcommands above are specified, assume this is a once-off stack trace dump
-			// @todo JohnB: This will also capture mistyped commands, so find a way to detect that
+			// @todo #JohnB: This will also capture mistyped commands, so find a way to detect that
 			else
 			{
 				GTraceManager.TraceAndDump(TraceName);
diff --git a/Engine/Plugins/NetcodeUnitTest/NetcodeUnitTest/Source/NetcodeUnitTest/Private/UnitTests/VMReflection.cpp b/Engine/Plugins/NetcodeUnitTest/NetcodeUnitTest/Source/NetcodeUnitTest/Private/UnitTests/VMReflection.cpp
--- a/Engine/Plugins/NetcodeUnitTest/NetcodeUnitTest/Source/NetcodeUnitTest/Private/UnitTests/VMReflection.cpp
+++ b/Engine/Plugins/NetcodeUnitTest/NetcodeUnitTest/Source/NetcodeUnitTest/Private/UnitTests/VMReflection.cpp
@@ -91,7 +91,7 @@ bool UVMReflection::ExecuteUnitTest()
 		TestResults.Add(TEXT("UObject Writing"), (!bError && TestObjA->AObjectRef == TargetResult));
 	}
 
-	// @todo JohnB: Add a test for UObject** casting, i.e. writable object property references, now that it is supported
+	// @todo #JohnB_VMRefl: Add a test for UObject** casting, i.e. writable object property references, now that it is supported
 
 
 
diff --git a/Engine/Plugins/NetcodeUnitTest/NetcodeUnitTest/Source/NetcodeUnitTest/Public/NUTUtil.h b/Engine/Plugins/NetcodeUnitTest/NetcodeUnitTest/Source/NetcodeUnitTest/Public/NUTUtil.h
--- a/Engine/Plugins/NetcodeUnitTest/NetcodeUnitTest/Source/NetcodeUnitTest/Public/NUTUtil.h
+++ b/Engine/Plugins/NetcodeUnitTest/NetcodeUnitTest/Source/NetcodeUnitTest/Public/NUTUtil.h
@@ -7,9 +7,9 @@
 class UUnitTest;
 
 
-// @todo JohnB: Convert to multicast delegate
+// @todo #JohnBRefactor: Convert to multicast delegate
 
-// @todo JohnB: Adjust all of these utility .h files, so that they implement code in .cpp, as these probably slow down compilation
+// @todo #JohnBRefactor: Adjust all of these utility .h files, so that they implement code in .cpp, as these probably slow down compilation
 
 #if !UE_BUILD_SHIPPING
 // Process even callback, which passes an extra parameter, for identifying the origin of the hook
diff --git a/Engine/Plugins/NetcodeUnitTest/NetcodeUnitTest/Source/NetcodeUnitTest/Public/NUTUtilDebug.h b/Engine/Plugins/NetcodeUnitTest/NetcodeUnitTest/Source/NetcodeUnitTest/Public/NUTUtilDebug.h
--- a/Engine/Plugins/NetcodeUnitTest/NetcodeUnitTest/Source/NetcodeUnitTest/Public/NUTUtilDebug.h
+++ b/Engine/Plugins/NetcodeUnitTest/NetcodeUnitTest/Source/NetcodeUnitTest/Public/NUTUtilDebug.h
@@ -38,7 +38,7 @@ IMPLEMENT_GET_PRIVATE_VAR(FStackTracker, bIsEnabled, bool);
  * NOTE: If you are trying to catch remote log messages deep within the internal game netcode, 
  *			then this class may not be appropriate, as remote logging passes through the netcode (YMMV)
  */
-class FScopedLog
+class NETCODEUNITTEST_API FScopedLog
 {
 protected:
 	FScopedLog()
@@ -88,14 +88,14 @@ protected:
 /**
  * Version of the above, for scoped logging of all netcode-related logs
  */
-class FScopedLogNet : public FScopedLog
+class NETCODEUNITTEST_API FScopedLogNet : public FScopedLog
 {
 public:
 	FScopedLogNet(UClientUnitTest* InUnitTest, bool bInRemoteLogging=false)
 	{
 		TArray<FString> TempLogCategories;
 
-		// @todo JohnB: See if there are any other good net categories to add here
+		// @todo #JohnBDebug: See if there are any other good net categories to add here
 		TempLogCategories.Add(TEXT("LogNet"));
 		TempLogCategories.Add(TEXT("LogRep"));
 		TempLogCategories.Add(TEXT("LogNetTraffic"));
@@ -115,7 +115,7 @@ public:
 /**
  * Base class for transparently hooking ProcessEvent
  */
-class FProcessEventHookBase
+class NETCODEUNITTEST_API FProcessEventHookBase
 {
 public:
 	FProcessEventHookBase();
@@ -145,7 +145,7 @@ protected:
 /**
  * A class for hooking and logging all ProcessEvent calls, within a particular code scope, similar to the above code
  */
-class FScopedProcessEventLog : FProcessEventHookBase
+class NETCODEUNITTEST_API FScopedProcessEventLog : FProcessEventHookBase
 {
 public:
 	FScopedProcessEventLog()
@@ -219,7 +219,7 @@ protected:
 };
 
 
-// @todo JohnB: Add documentation for all of the stack trace stuff, to the primary NetcodeUnitTest documentation
+// @todo #JohnBDocMain: Add documentation for all of the stack trace stuff, to the primary NetcodeUnitTest documentation
 
 /**
  * Manager for handling multiple debug stack traces on-the-fly, and allowing abstraction of stack traces,
@@ -361,11 +361,11 @@ protected:
 };
 
 
-// @todo JohnB: The log trace stuff doesn't support the 'startenabled' and enable/disable etc. features;
+// @todo #JohnBDebug: The log trace stuff doesn't support the 'startenabled' and enable/disable etc. features;
 //				some stuff (like enable/disable) should be easy to add though, so do that eventually
 
 
-// @todo JohnB: NOTE: Don't move the log stack trace manager to .cpp yet; not until you decide how to implement,
+// @todo #JohnB: NOTE: Don't move the log stack trace manager to .cpp yet; not until you decide how to implement,
 //				the similar ProcessEvent hook that will tie in to the trace manager, as you may try to generalize the below for that
 
 /**
@@ -605,7 +605,7 @@ public:
 };
 
 
-// @todo JohnB: Continue implementing
+// @todo #JohnBFeatureDebug: Continue implementing
 #if 0
 #if !UE_BUILD_SHIPPING
 /**
@@ -636,7 +636,7 @@ public:
  */
 struct NUTDebug
 {
-	// @todo JohnB: Might be useful to find a place for the hexdump functions within the base engine code itself;
+	// @todo #JohnBMerge: Might be useful to find a place for the hexdump functions within the base engine code itself;
 	//				they are very useful
 
 	// NOTE: If outputting hex dump info to the unit test log windows, make use of bMonospace, to retain the hex formatting
diff --git a/Engine/Plugins/NetcodeUnitTest/NetcodeUnitTest/Source/NetcodeUnitTest/Public/NUTUtilProfiler.h b/Engine/Plugins/NetcodeUnitTest/NetcodeUnitTest/Source/NetcodeUnitTest/Public/NUTUtilProfiler.h
--- a/Engine/Plugins/NetcodeUnitTest/NetcodeUnitTest/Source/NetcodeUnitTest/Public/NUTUtilProfiler.h
+++ b/Engine/Plugins/NetcodeUnitTest/NetcodeUnitTest/Source/NetcodeUnitTest/Public/NUTUtilProfiler.h
@@ -74,8 +74,11 @@ protected:
 	/** Whether or not profiling is active */
 	bool bActive;
 
+
+#if TARGET_UE4_CL >= CL_DEPRECATEDEL
 	/** Handle to the registered OnNewFrame delegate */
 	FDelegateHandle OnNewFrameDelegateHandle;
+#endif
 };
 #endif
 
diff --git a/Engine/Plugins/NetcodeUnitTest/NetcodeUnitTest/Source/NetcodeUnitTest/Public/NUTUtilReflection.h b/Engine/Plugins/NetcodeUnitTest/NetcodeUnitTest/Source/NetcodeUnitTest/Public/NUTUtilReflection.h
--- a/Engine/Plugins/NetcodeUnitTest/NetcodeUnitTest/Source/NetcodeUnitTest/Public/NUTUtilReflection.h
+++ b/Engine/Plugins/NetcodeUnitTest/NetcodeUnitTest/Source/NetcodeUnitTest/Public/NUTUtilReflection.h
@@ -145,19 +145,19 @@
  */
 
 
-// @todo JohnB: Write unit test for 'struct object access' above (a struct with an object member)
-// @todo JohnB: Add 'structonscope' example to unit test as well (it's verified as working, but should be in the unit test too)
+// @todo #JohnB_VMRefl: Write unit test for 'struct object access' above (a struct with an object member)
+// @todo #JohnB_VMRefl: Add 'structonscope' example to unit test as well (it's verified as working, but should be in the unit test too)
 
 
-// @todo JohnB: Add an example of iterating struct arrays of unknown struct type, to the unit test - as an example, see:
+// @todo #JohnB_VMRefl: Add an example of iterating struct arrays of unknown struct type, to the unit test - as an example, see:
 //				UInfiniteItems_Grenades::GetInventoryInfo
 
-// @todo JohnB: You haven't yet had to actually adjust the size of dynamic arrays of unknown/inaccessible type,
+// @todo #JohnB_VMRefl: You haven't yet had to actually adjust the size of dynamic arrays of unknown/inaccessible type,
 //				or add/remove elements and other manipulation yet; add and document a method of doing this, eventually
 
 
 
-// @todo JohnB: In principle, there's nothing to stop you adding another boilerplate 'FVMReflectionString' class,
+// @todo #JohnBFeature: In principle, there's nothing to stop you adding another boilerplate 'FVMReflectionString' class,
 //				taking and parsing a string using same syntax as FVMReflection (may as well optimize syntax though, as C++ is limited),
 //				then tying that to a console-command, as a supercharged version of 'get/set' commands, but with reflection capabilities.
 //
@@ -684,7 +684,7 @@ struct NUTUtilRefl
 	/**
 	 * Iterates a functions parameters, and converts them to a human readable string
 	 * NOTE: Can not be passed back to functions that read parameters from a string, e.g. 'CallFunctionByNameWithArguments'
-	 *			@todo JohnB: This would be a good goal though
+	 *			@todo #JohnB: This would be a good goal though
 	 *
 	 * @param InFunction	The function template whose parameters are being parsed
 	 * @param Parms			The parameter data to be parsed
diff --git a/Engine/Plugins/NetcodeUnitTest/NetcodeUnitTest/Source/NetcodeUnitTest/Public/Net/NUTUtilNet.h b/Engine/Plugins/NetcodeUnitTest/NetcodeUnitTest/Source/NetcodeUnitTest/Public/Net/NUTUtilNet.h
--- a/Engine/Plugins/NetcodeUnitTest/NetcodeUnitTest/Source/NetcodeUnitTest/Public/Net/NUTUtilNet.h
+++ b/Engine/Plugins/NetcodeUnitTest/NetcodeUnitTest/Source/NetcodeUnitTest/Public/Net/NUTUtilNet.h
@@ -107,8 +107,13 @@ public:
 	{
 		if (AttachedWorld != NULL)
 		{
+#if TARGET_UE4_CL >= CL_DEPRECATEDEL
 			TickDispatchDelegateHandle  = AttachedWorld->OnTickDispatch().AddRaw(this, &FWorldTickHook::TickDispatch);
 			PostTickFlushDelegateHandle = AttachedWorld->OnPostTickFlush().AddRaw(this, &FWorldTickHook::PostTickFlush);
+#else
+			AttachedWorld->OnTickDispatch().AddRaw(this, &FWorldTickHook::TickDispatch);
+			AttachedWorld->OnPostTickFlush().AddRaw(this, &FWorldTickHook::PostTickFlush);
+#endif
 		}
 	}
 
@@ -116,8 +121,13 @@ public:
 	{
 		if (AttachedWorld != NULL)
 		{
+#if TARGET_UE4_CL >= CL_DEPRECATEDEL
 			AttachedWorld->OnPostTickFlush().Remove(PostTickFlushDelegateHandle);
 			AttachedWorld->OnTickDispatch().Remove(TickDispatchDelegateHandle);
+#else
+			AttachedWorld->OnPostTickFlush().RemoveRaw(this, &FWorldTickHook::PostTickFlush);
+			AttachedWorld->OnTickDispatch().RemoveRaw(this, &FWorldTickHook::TickDispatch);
+#endif
 		}
 
 		AttachedWorld = NULL;
@@ -138,12 +148,14 @@ public:
 	/** The world this is attached to */
 	UWorld* AttachedWorld;
 
+#if TARGET_UE4_CL >= CL_DEPRECATEDEL
 private:
 	/** Handle for Tick dispatch delegate */
 	FDelegateHandle TickDispatchDelegateHandle;
 
 	/** Handle for PostTick dispatch delegate */
 	FDelegateHandle PostTickFlushDelegateHandle;
+#endif
 };
 
 /**
diff --git a/Engine/Plugins/NetcodeUnitTest/NetcodeUnitTest/Source/NetcodeUnitTest/Public/NetcodeUnitTest.h b/Engine/Plugins/NetcodeUnitTest/NetcodeUnitTest/Source/NetcodeUnitTest/Public/NetcodeUnitTest.h
--- a/Engine/Plugins/NetcodeUnitTest/NetcodeUnitTest/Source/NetcodeUnitTest/Public/NetcodeUnitTest.h
+++ b/Engine/Plugins/NetcodeUnitTest/NetcodeUnitTest/Source/NetcodeUnitTest/Public/NetcodeUnitTest.h
@@ -6,19 +6,20 @@
  * Compatibility defines
  */
 
-// @todo JohnB: Perhaps move compatibility defines to its own header
+// @todo #JohnBRefactor: Perhaps move compatibility defines to its own header
 
-// @todo JohnB: Every time you update NetcodeUnitTest code, to toggle based on a changelist, add a define for it with the CL number here
-
-// @todo JohnB: I need accurate mainline CL numbers for UT integrations (they aren't included in the merge changelist descriptions)
+// @todo #JohnB: I need accurate mainline CL numbers for UT integrations (they aren't included in the merge changelist descriptions)
 
 
 /** List of mainline UE4 changelists, representing UE4 version changes (as judged by changes to Version.h) */
-#define CL_4_8 2388573
-#define CL_4_7 2347015
-#define CL_4_6 2308471
+#define CL_4_10	2626674
+#define CL_4_9	2526821
+#define CL_4_8	2388573
+#define CL_4_7	2347015
+#define CL_4_6	2308471
 
 /** List of UnrealTournament mainline merges (number represents merged mainline CL) NOTE: Estimated. No integrated CL number in merge */
+#define CL_UT_4_8 CL_4_9	// Integrated with final 4.8 release, which is approximately same time as main branch started 4.9
 #define CL_UT_4_7 CL_4_7
 #define CL_UT_4_6 CL_4_6
 
@@ -28,7 +29,11 @@
 #define CL_FORT_4_7 2349525
 
 
-/** List of mainline UE4 engine changelists, that required a NetcodeUnitTest compatibility adjustment (newest first) */
+/**
+ * List of mainline UE4 engine changelists, that required a NetcodeUnitTest compatibility adjustment (newest first)
+ * Every time you update NetcodeUnitTest code, to toggle based on a changelist, add a define for it with the CL number here
+ */
+#define CL_INITCONNPARAM	2567692
 #define CL_CONSTUNIQUEID	2540329
 #define CL_CONSTNETCONN		2501704
 #define CL_INPUTCHORD		2481648
@@ -37,6 +42,7 @@
 #define CL_BEACONHOST		2456855
 #define CL_GETSELECTIONMODE	2425976
 #define CL_DEPRECATENEW		2425600
+#define CL_DEPRECATEDEL		2400883
 #define CL_FNETWORKVERSION	2384479
 #define CL_CHECKBOXSTATE	2384008
 
@@ -49,7 +55,7 @@
  *
  * If in doubt, set to the top CL from list above ('List of mainline UE4 engine changelists'), and work your way down until it compiles.
  */
-#define TARGET_UE4_CL 2554937	// IMPORTANT: If you're hovering over this because compiling failed, you need to adjust this value.
+#define TARGET_UE4_CL 2619427	// IMPORTANT: If you're hovering over this because compiling failed, you need to adjust this value.
 
 
 /**
@@ -60,18 +66,14 @@ class UUnitTestManager;
 class UUnitTestBase;
 
 
-// @todo JohnB: Status window logs (such as unit test success/failure), don't appear to add the same colouring to main game log window;
+// @todo #JohnBLowPri: Status window logs (such as unit test success/failure), don't appear to add the same colouring to main game log window;
 //				add this sometime (probably based on how FUnitTestProcess does it)
 
 /**
  * Enums
  */
 
-// @todo JohnB: Also consider trying to refactor log-colouring, to consider these flags (or even do-away with all log colouring code,
-//				that is NOT based on these flags)
-
-
-// @todo JohnB: Trim values from the below enum - not all are necessary
+// @todo #JohnBRefactor: Trim values from the below enum - not all are necessary
 
 /**
  * Used to help identify what type of log is being processed
@@ -143,7 +145,7 @@ extern UUnitTestManager* GUnitTestManager;
 /** Used by to aid in hooking log messages triggered by unit tests */
 extern NETCODEUNITTEST_API UUnitTestBase* GActiveLogUnitTest;
 extern UUnitTestBase* GActiveLogEngineEvent;
-extern UWorld* GActiveLogWorld;				// @todo JohnB: Does this make other log hooks redundant?
+extern UWorld* GActiveLogWorld;				// @todo #JohnBReview: Does this make other log hooks redundant?
 
 /** Keeps track of the UUnitTestNetConnection, currently processing received data (to aid with selective log hooking) */
 extern UNetConnection* GActiveReceiveUnitConnection;
@@ -192,7 +194,7 @@ inline ELogType OptionalFlags(ELogType InFlags=ELogType::None)
 }
 
 // Actual asserts (not optimized out, like almost all other engine assert macros)
-// @todo JohnB: Try to get this to show up a message box, or some other notification, rather than just exiting immediately
+// @todo #JohnBLowPri: Try to get this to show up a message box, or some other notification, rather than just exiting immediately
 #define UNIT_ASSERT(Condition) \
 	if (!(Condition)) \
 	{ \
@@ -264,7 +266,7 @@ inline ELogType OptionalFlags(ELogType InFlags=ELogType::None)
 	UNIT_LOG_OBJ(this, UnitLogTypeFlags, Format, __VA_ARGS__)
 
 
-// @todo JohnB: Document these macros
+// @todo #JohnBDoc: Document these macros
 
 #define UNIT_EVENT_BEGIN(UnitTestObj) \
 	UUnitTestBase* OldEventLogVal = GActiveLogEngineEvent; \
diff --git a/Engine/Plugins/NetcodeUnitTest/NetcodeUnitTest/Source/NetcodeUnitTest/Public/UI/LogWindowManager.h b/Engine/Plugins/NetcodeUnitTest/NetcodeUnitTest/Source/NetcodeUnitTest/Public/UI/LogWindowManager.h
--- a/Engine/Plugins/NetcodeUnitTest/NetcodeUnitTest/Source/NetcodeUnitTest/Public/UI/LogWindowManager.h
+++ b/Engine/Plugins/NetcodeUnitTest/NetcodeUnitTest/Source/NetcodeUnitTest/Public/UI/LogWindowManager.h
@@ -3,7 +3,7 @@
 #pragma once
 
 /**
- * @todo JohnB: Slate log windows
+ * @todo #JohnBFeatureUI: Slate log windows
  *
  * What is needed:
  *	- Ability to enter console commands into log window (what context should they execute in? game client, exploit, server?)
@@ -137,8 +137,11 @@ protected:
 	int									LogWidth;
 	int									LogHeight;
 
+
+#if TARGET_UE4_CL >= CL_DEPRECATEDEL
 	/** Handles to registered OnWindowClosed delegates for particular windows */
 	TMap<SLogWindow*, FDelegateHandle> OnWindowClosedDelegateHandles;
+#endif
 };
 
 
diff --git a/Engine/Plugins/NetcodeUnitTest/NetcodeUnitTest/Source/NetcodeUnitTest/Public/UI/SLogWidget.h b/Engine/Plugins/NetcodeUnitTest/NetcodeUnitTest/Source/NetcodeUnitTest/Public/UI/SLogWidget.h
--- a/Engine/Plugins/NetcodeUnitTest/NetcodeUnitTest/Source/NetcodeUnitTest/Public/UI/SLogWidget.h
+++ b/Engine/Plugins/NetcodeUnitTest/NetcodeUnitTest/Source/NetcodeUnitTest/Public/UI/SLogWidget.h
@@ -104,7 +104,7 @@ struct FLogTabInfo
 	TWeakPtr<SDockTab> TabWidget;
 
 
-	// @todo JohnB: Add an extended filter eventually (perhaps an FString for simple String.Contains filter?)
+	// @todo #JohnBFeatureUI: Add an extended filter eventually (perhaps an FString for simple String.Contains filter?)
 
 
 	FLogTabInfo(FString InLabel, FString InToolTip, ELogType InFilter=ELogType::All, uint8 InPriority=255, bool bInTabOpen=true)
@@ -158,7 +158,6 @@ public:
 		, OnDeveloperClicked()
 		, OnConsoleCommand()
 	{
-		// @todo JohnB: Compile a list of unit test commandline parameters somewhere (maybe in the commandlet header)
 		bAutoClose = !FParse::Param(FCommandLine::Get(), TEXT("UnitTestNoAutoClose"));
 	}
 
diff --git a/Engine/Plugins/NetcodeUnitTest/NetcodeUnitTest/VAXSnippets/Unit Test Header.txt b/Engine/Plugins/NetcodeUnitTest/NetcodeUnitTest/VAXSnippets/Unit Test Header.txt
--- a/Engine/Plugins/NetcodeUnitTest/NetcodeUnitTest/VAXSnippets/Unit Test Header.txt	
+++ b/Engine/Plugins/NetcodeUnitTest/NetcodeUnitTest/VAXSnippets/Unit Test Header.txt	
@@ -12,7 +12,7 @@ Source:
 #include "$FILE_BASE$.generated.h"
 
 /**
- * @todo JohnB: $end$
+ * @todo #JohnB: $end$
  */
 UCLASS()
 class U$FILE_BASE$ : public UClientUnitTest
diff --git a/Engine/Plugins/NetcodeUnitTest/NetcodeUnitTest/VAXSnippets/Unit Test Source.txt b/Engine/Plugins/NetcodeUnitTest/NetcodeUnitTest/VAXSnippets/Unit Test Source.txt
--- a/Engine/Plugins/NetcodeUnitTest/NetcodeUnitTest/VAXSnippets/Unit Test Source.txt	
+++ b/Engine/Plugins/NetcodeUnitTest/NetcodeUnitTest/VAXSnippets/Unit Test Source.txt	
@@ -26,7 +26,7 @@ U$FILE_BASE$::U$FILE_BASE$(const FObjectInitializer& ObjectInitializer)
 
 	bWorkInProgress = true;
 
-	// @todo JohnB: Bugtracking/changelist notes
+	// @todo #JohnBExploitCL: Bugtracking/changelist notes
 
 	ExpectedResult.Add(TEXT("$TargetGame$"), EUnitTestVerification::VerifiedNotFixed);
 
@@ -44,7 +44,7 @@ void U$FILE_BASE$::InitializeEnvironmentSettings()
 
 void U$FILE_BASE$::ExecuteClientUnitTest()
 {
-	// @todo JohnB
+	// @todo #JohnB
 }
 
 void U$FILE_BASE$::NotifyProcessLog(TWeakPtr<FUnitTestProcess> InProcess, const TArray<FString>& InLogLines)
@@ -55,7 +55,7 @@ void U$FILE_BASE$::NotifyProcessLog(TWeakPtr<FUnitTestProcess> InProcess, const
 	{
 		for (auto CurLine : InLogLines)
 		{
-			// @todo JohnB
+			// @todo #JohnB
 		}
 	}
 }
