NetcodeUnitTest disabling code for VS2015

diff --git a/Engine/Plugins/NetcodeUnitTest/NetcodeUnitTest/Source/NetcodeUnitTest/Private/UI/SLogWidget.cpp b/Engine/Plugins/NetcodeUnitTest/NetcodeUnitTest/Source/NetcodeUnitTest/Private/UI/SLogWidget.cpp
--- a/Engine/Plugins/NetcodeUnitTest/NetcodeUnitTest/Source/NetcodeUnitTest/Private/UI/SLogWidget.cpp
+++ b/Engine/Plugins/NetcodeUnitTest/NetcodeUnitTest/Source/NetcodeUnitTest/Private/UI/SLogWidget.cpp
@@ -30,7 +30,7 @@ IMPLEMENT_GET_PRIVATE_VAR(SEditableTextBox, EditableText, TSharedPtr<SEditableTe
 IMPLEMENT_GET_PRIVATE_VAR(SButton, Style, const FButtonStyle*);
 
 // Enable access to SDockTab::GetCurrentStyle, using the CALL_PRIVATE macro
-IMPLEMENT_GET_PRIVATE_FUNC_CONST(SDockTab, GetCurrentStyle, const FDockTabStyle&, void, const);
+// IMPLEMENT_GET_PRIVATE_FUNC_CONST(SDockTab, GetCurrentStyle, const FDockTabStyle&, void, const);
 
 
 // @todo #JohnBRefactorUI: Perhaps move widget searching to a NUTSlate.h file, or such?
@@ -968,9 +968,9 @@ TSharedRef<SDockTab> SLogWidget::SpawnLogTab(const FSpawnTabArgs& InSpawnTabArgs
 				{
 					TSharedRef<SButton> CurButton = StaticCastSharedRef<SButton>(InWidget);
 					const FButtonStyle* ButtonStyle = GET_PRIVATE(SButton, CurButton, Style);
-					const FDockTabStyle& TabStyle = CALL_PRIVATE(SDockTab, ReturnVal, GetCurrentStyle)();
-
-					bFound = ButtonStyle == &TabStyle.CloseButtonStyle;
+// 					const FDockTabStyle& TabStyle = CALL_PRIVATE(SDockTab, ReturnVal, GetCurrentStyle)();
+// 
+// 					bFound = ButtonStyle == &TabStyle.CloseButtonStyle;
 				}
 
 				return bFound;
diff --git a/Engine/Plugins/NetcodeUnitTest/NetcodeUnitTest/Source/NetcodeUnitTest/Private/UnitTests/VMReflection.cpp b/Engine/Plugins/NetcodeUnitTest/NetcodeUnitTest/Source/NetcodeUnitTest/Private/UnitTests/VMReflection.cpp
--- a/Engine/Plugins/NetcodeUnitTest/NetcodeUnitTest/Source/NetcodeUnitTest/Private/UnitTests/VMReflection.cpp
+++ b/Engine/Plugins/NetcodeUnitTest/NetcodeUnitTest/Source/NetcodeUnitTest/Private/UnitTests/VMReflection.cpp
@@ -774,7 +774,7 @@ bool UVMReflection::ExecuteUnitTest()
 		}
 
 		bool bError = false;
-		FText Result = (FText)(FVMReflection(TestObj)->*"TextProp", &bError);
+		FText Result = FText::GetEmpty();// (FText)((FVMReflection(TestObj)->*"TextProp"), &bError);
 
 		TestResults.Add(TEXT("FText Reading"), (!bError && Result.EqualTo(TargetResult)));
 	}
@@ -1364,7 +1364,7 @@ bool UVMReflection::ExecuteUnitTest()
 		bool bError = false;
 		FText Result;
 		
-		Result = (FText)((FVMReflection(TestObj)->*"DynTextPropArray")["FText"][0], &bError);
+		Result = FText::GetEmpty();// (FText)((FVMReflection(TestObj)->*"DynTextPropArray")["FText"][0], &bError);
 
 		TestResults.Add(TEXT("Dynamic Array Type Verify (FText)"), (!bError && Result.EqualTo(TargetResult)));
 	}
