Added better detection of shader compilation (bad for unit tests), and fixed a race condition which randomly broke some unit tests - along with general refactoring.

Full set of changes:
- Refactored all child-process related unit test code, into a new ProcessUnitTest class, to allow use with non-UE4 processes
- Added better detection of ShaderCompileWorker and shader compilation in general, and made unit tests more tolerant of this (but spam an error message, stating that shaders should be built independently)
- Fixed race condition, where NUTActor control-channel messages, could be executed without an Owner - randomly breaking some unit tests

diff --git a/Engine/Plugins/NetcodeUnitTest/NetcodeUnitTest/Source/NetcodeUnitTest/Classes/ClientUnitTest.h b/Engine/Plugins/NetcodeUnitTest/NetcodeUnitTest/Source/NetcodeUnitTest/Classes/ClientUnitTest.h
--- a/Engine/Plugins/NetcodeUnitTest/NetcodeUnitTest/Source/NetcodeUnitTest/Classes/ClientUnitTest.h
+++ b/Engine/Plugins/NetcodeUnitTest/NetcodeUnitTest/Source/NetcodeUnitTest/Classes/ClientUnitTest.h
@@ -2,7 +2,7 @@
 
 #pragma once
 
-#include "UnitTest.h"
+#include "ProcessUnitTest.h"
 
 #include "ClientUnitTest.generated.h"
 
@@ -132,136 +132,6 @@ inline FString GetUnitTestFlagName(EUnitTestFlags Flag)
 }
 
 
-/**
- * Enum for different stages of error log parsing
- */
-enum class EErrorLogStage : uint8
-{
-	ELS_NoError,		// No error logs have been received/parsed yet
-	ELS_ErrorStart,		// The text indicating the start of an error log is being parsed
-	ELS_ErrorDesc,		// The text describing the error is being parsed
-	ELS_ErrorCallstack,	// The callstack for the error is being parsed
-	ELS_ErrorExit		// The post-error exit message is being parsed (error parsing is effectively complete)
-};
-
-/**
- * Enum for specifying the suspend state of a process (typically the server)
- */
-enum class ESuspendState : uint8
-{
-	Active,				// Process is currently active/not-suspended
-	Suspended,			// Process is currently suspended
-};
-
-
-// Delegate definitions
-
-/**
- * Delegate notifying that the server suspend state has changed
- *
- * @param NewSuspendState	The new server suspend state
- */
-DECLARE_DELEGATE_OneParam(FOnSuspendStateChange, ESuspendState /*NewSuspendState*/);
-
-
-// Struct definitions
-
-/**
- * Struct used for storing and classifying each log error line
- */
-struct FErrorLog
-{
-	/** The stage of this error log line */
-	EErrorLogStage Stage;
-
-	/** The error log line */
-	FString Line;
-
-
-	/**
-	 * Base constructor
-	 */
-	FErrorLog()
-		: Stage(EErrorLogStage::ELS_NoError)
-		, Line()
-	{
-	}
-
-	FErrorLog(EErrorLogStage InStage, const FString& InLine)
-		: Stage(InStage)
-		, Line(InLine)
-	{
-	}
-};
-
-/**
- * Struct used for handling a launched UE4 client/server process
- */
-struct FUnitTestProcess
-{
-	friend class UClientUnitTest;
-
-	/** Process handle for the launched process */
-	FProcHandle ProcessHandle;
-
-	/** The process ID */
-	uint32 ProcessID;
-
-	/** The suspend state of the process (implemented as a part of unit test code, does not relate to OS API) */
-	ESuspendState SuspendState;
-
-	/** Human-readable tag given to this process */
-	FString ProcessTag;
-
-
-	/** Handle to StdOut for the launched process */
-	void* ReadPipe;
-
-	/** Handle to StdIn for the launched process (unused) */
-	void* WritePipe;
-
-
-	/** The base log type for this process (client? server? process?) */
-	ELogType BaseLogType;
-
-	/** The prefix to use for StdOut log output */
-	FString LogPrefix;
-
-	/** The OutputDeviceColor string, to use for setting the log output color */
-	const TCHAR* MainLogColor;
-
-	/** The log output color to use in the slate log window */
-	FSlateColor SlateLogColor;
-
-
-	/** If this process is outputting an error log, this is the current stage of error parsing (or ELS_NoError if not parsing) */
-	EErrorLogStage ErrorLogStage;
-
-	/** Gathered error log text */
-	TArray<FErrorLog> ErrorText;
-
-
-	/**
-	 * Base constructor
-	 */
-	FUnitTestProcess()
-		: ProcessHandle()
-		, ProcessID(0)
-		, SuspendState(ESuspendState::Active)
-		, ProcessTag(TEXT(""))
-		, ReadPipe(NULL)
-		, WritePipe(NULL)
-		, BaseLogType(ELogType::None)
-		, LogPrefix(TEXT(""))
-		, MainLogColor(COLOR_NONE)
-		, SlateLogColor(FSlateColor::UseForeground())
-		, ErrorLogStage(EErrorLogStage::ELS_NoError)
-		, ErrorText()
-	{
-	}
-};
-
-
 /**
  * Base class for all unit tests depending upon a (fake/minimal) client connecting to a server.
  * Handles creation/cleanup of an entire new UWorld, UNetDriver and UNetConnection, for fast sequential unit testing.
@@ -269,7 +139,7 @@ struct FUnitTestProcess
  * In subclasses, implement the unit test within the ExecuteClientUnitTest function (remembering to call parent)
  */
 UCLASS()
-class NETCODEUNITTEST_API UClientUnitTest : public UUnitTest
+class NETCODEUNITTEST_API UClientUnitTest : public UProcessUnitTest
 {
 	GENERATED_UCLASS_BODY()
 
@@ -309,9 +179,6 @@ protected:
 
 	/** Runtime variables */
 protected:
-	/** Stores a reference to all running child processes tied to this unit test, for housekeeping */
-	TArray<TSharedPtr<FUnitTestProcess>> ActiveProcesses;
-
 	/** Reference to the created server process handling struct */
 	TWeakPtr<FUnitTestProcess> ServerHandle;
 
@@ -324,6 +191,18 @@ protected:
 	/** Reference to the created client process handling struct (if enabled) */
 	TWeakPtr<FUnitTestProcess> ClientHandle;
 
+	/** Whether or not there is a blocking event/process preventing setup of the server */
+	bool bBlockingServerDelay;
+
+	/** Whether or not there is a blocking event/process preventing setup of a client */
+	bool bBlockingClientDelay;
+
+	/** Whether or not there is a blocking event/process preventing the fake client from connecting */
+	bool bBlockingFakeClientDelay;
+
+	/** When a server is launched after a blocking event/process, this delays the launch of any clients, in case of more blockages */
+	double NextBlockingTimeout;
+
 
 	/** Stores a reference to the created fake world, for execution and later cleanup */
 	UWorld* UnitWorld;
@@ -359,11 +238,6 @@ protected:
 	TArray<int32> PendingNetActorChans;
 
 
-public:
-	/** Delegate for notifying the UI, of a change in the server suspend state */
-	FOnSuspendStateChange OnServerSuspendState;
-
-
 #if TARGET_UE4_CL >= CL_DEPRECATEDEL
 private:
 	/** Handle to the registered InternalNotifyNetworkFailure delegate */
@@ -497,55 +371,19 @@ public:
 	virtual bool NotifySendRPC(AActor* Actor, UFunction* Function, void* Parameters, FOutParmRec* OutParms, FFrame* Stack,
 								UObject* SubObject);
 
-	/**
-	 * For implementation in subclasses, for helping to verify success/fail upon completion of unit tests
-	 * NOTE: Not called again once VerificationState is set
-	 * WARNING: Be careful when iterating InLogLines in multiple different for loops, if the sequence of detected logs is important
-	 *
-	 * @param InProcess		The process the log lines are from
-	 * @param InLogLines	The current log lines being received
-	 */
-	virtual void NotifyProcessLog(TWeakPtr<FUnitTestProcess> InProcess, const TArray<FString>& InLogLines);
 
-	/**
-	 * Notifies that there was a request to suspend/resume the unit test server
-	 */
-	void NotifySuspendRequest();
+	virtual void NotifyProcessLog(TWeakPtr<FUnitTestProcess> InProcess, const TArray<FString>& InLogLines) override;
 
-	/**
-	 * Notifies when the suspend state of a process changes
-	 *
-	 * @param InProcess			The process whose suspend state has changed
-	 * @param InSuspendState	The new suspend state of the process
-	 */
-	void NotifyProcessSuspendState(TWeakPtr<FUnitTestProcess> InProcess, ESuspendState InSuspendState);
+	virtual void NotifyProcessFinished(TWeakPtr<FUnitTestProcess> InProcess) override;
 
+	virtual void NotifyProcessSuspendState(TWeakPtr<FUnitTestProcess> InProcess, ESuspendState InSuspendState) override;
 
-	/**
-	 * Notifies that there was a request to enable/disable developer mode
-	 *
-	 * @param bInDeveloperMode	Whether or not developer mode is being enabled/disabled
-	 */
-	void NotifyDeveloperModeRequest(bool bInDeveloperMode);
 
-	/**
-	 * Notifies that there was a request to execute a console command for the unit test, which can occur in a specific context,
-	 * e.g. for the unit test server, for the local minimal-client (within the unit test), or for the separate unit test client process
-	 *
-	 * @param CommandContext	The context (local/server/client?) for the console command
-	 * @param Command			The command to be executed
-	 * @return					Whether or not the command was handled
-	 */
-	virtual bool NotifyConsoleCommandRequest(FString CommandContext, FString Command);
+	virtual void NotifySuspendRequest() override;
 
+	virtual bool NotifyConsoleCommandRequest(FString CommandContext, FString Command) override;
 
-	/**
-	 * Outputs the list of console command contexts, that this unit test supports (which can include custom contexts in subclasses)
-	 *
-	 * @param OutList				Outputs the list of supported console command contexts
-	 * @param OutDefaultContext		Outputs the context which should be auto-selected/defaulted-to
-	 */
-	virtual void GetCommandContextList(TArray<TSharedPtr<FString>>& OutList, FString& OutDefaultContext);
+	virtual void GetCommandContextList(TArray<TSharedPtr<FString>>& OutList, FString& OutDefaultContext) override;
 
 
 	/**
@@ -609,8 +447,10 @@ protected:
 
 	/**
 	 * Resets the net connection timeout
+	 *
+	 * @param Duration	The duration which the timeout reset should last
 	 */
-	void ResetConnTimeout();
+	void ResetConnTimeout(float Duration);
 
 
 	/**
@@ -698,41 +538,7 @@ protected:
 	virtual FString ConstructClientParameters(FString ConnectIP);
 
 
-	/**
-	 * Starts a child UE4 process, tied to the unit test, with the specified commandline
-	 *
-	 * @param InCommandline		The commandline that the child process should use
-	 * @param bMinimized		Starts the process with the window minimized
-	 * @return					Returns a pointer to the new processes handling struct
-	 */
-	virtual TWeakPtr<FUnitTestProcess> StartUnitTestProcess(FString InCommandline, bool bMinimized=true);
-
-	/**
-	 * Shuts-down/cleans-up a child process tied to the unit test
-	 *
-	 * @param InHandle	The handling struct for the process
-	 */
-	virtual void ShutdownUnitTestProcess(TSharedPtr<FUnitTestProcess> InHandle);
-
-	/**
-	 * Processes the standard output (i.e. log output) for processes
-	 */
-	void PollProcessOutput();
-
-	/**
-	 * Updates (and if necessary, saves) the memory stats for processes
-	 */
-	void UpdateProcessStats();
-
-
-	/**
-	 * Checks incoming process logs, for any indication of a UE4 crash/error
-	 *
-	 * @param InProcess		The process the log output originates from
-	 * @param InLines		The incoming log lines
-	 */
-	void CheckOutputForError(TSharedPtr<FUnitTestProcess> InProcess, const TArray<FString>& InLines);
-
+	virtual void PrintUnitTestProcessErrors(TSharedPtr<FUnitTestProcess> InHandle) override;
 
 	void InternalNotifyNetworkFailure(UWorld* InWorld, UNetDriver* InNetDriver, ENetworkFailure::Type FailureType,
 										const FString& ErrorString);
@@ -743,15 +549,7 @@ protected:
 #endif
 
 
-	// Tick standard output of the server, if the process is launched
 	virtual void UnitTick(float DeltaTime) override;
 
-	virtual void PostUnitTick(float DeltaTime) override;
-
 	virtual bool IsTickable() const override;
-
-
-	virtual void FinishDestroy() override;
-
-	virtual void ShutdownAfterError() override;
 };
diff --git a/Engine/Plugins/NetcodeUnitTest/NetcodeUnitTest/Source/NetcodeUnitTest/Classes/NUTActor.h b/Engine/Plugins/NetcodeUnitTest/NetcodeUnitTest/Source/NetcodeUnitTest/Classes/NUTActor.h
--- a/Engine/Plugins/NetcodeUnitTest/NetcodeUnitTest/Source/NetcodeUnitTest/Classes/NUTActor.h
+++ b/Engine/Plugins/NetcodeUnitTest/NetcodeUnitTest/Source/NetcodeUnitTest/Classes/NUTActor.h
@@ -184,10 +184,12 @@ public:
 	void ServerExecute(const FString& InDelegate);
 
 
+	virtual void Tick(float DeltaSeconds) override;
+
 	/**
-	 * Monitor for initial player on server, in order to assign ownership
+	 * Update the owner for the NUTActor, if the current owner is no longer valid
 	 */
-	virtual void Tick(float DeltaSeconds) override;
+	void UpdateOwner();
 };
 
 
diff --git a/Engine/Plugins/NetcodeUnitTest/NetcodeUnitTest/Source/NetcodeUnitTest/Classes/Net/UnitTestPackageMap.h b/Engine/Plugins/NetcodeUnitTest/NetcodeUnitTest/Source/NetcodeUnitTest/Classes/Net/UnitTestPackageMap.h
--- a/Engine/Plugins/NetcodeUnitTest/NetcodeUnitTest/Source/NetcodeUnitTest/Classes/Net/UnitTestPackageMap.h
+++ b/Engine/Plugins/NetcodeUnitTest/NetcodeUnitTest/Source/NetcodeUnitTest/Classes/Net/UnitTestPackageMap.h
@@ -23,7 +23,7 @@ class UUnitTestPackageMap : public UPackageMapClient
 #endif
 
 
-	virtual bool SerializeObject(FArchive& Ar, UClass* Class, UObject*& Obj, FNetworkGUID* OutNetGUID=NULL) override;
+	virtual bool SerializeObject(FArchive& Ar, UClass* InClass, UObject*& Obj, FNetworkGUID* OutNetGUID=NULL) override;
 
 	virtual bool SerializeNewActor(FArchive& Ar, class UActorChannel* Channel, class AActor*& Actor) override;
 
diff --git a/Engine/Plugins/NetcodeUnitTest/NetcodeUnitTest/Source/NetcodeUnitTest/Classes/ProcessUnitTest.h b/Engine/Plugins/NetcodeUnitTest/NetcodeUnitTest/Source/NetcodeUnitTest/Classes/ProcessUnitTest.h
new file mode 100644
--- /dev/null
+++ b/Engine/Plugins/NetcodeUnitTest/NetcodeUnitTest/Source/NetcodeUnitTest/Classes/ProcessUnitTest.h
@@ -0,0 +1,282 @@
+// Copyright 1998-2015 Epic Games, Inc. All Rights Reserved.
+
+#pragma once
+
+#include "UnitTest.h"
+
+#include "ProcessUnitTest.generated.h"
+
+
+// Enum definitions
+
+/**
+ * Enum for different stages of error log parsing (mainly applicable to UE4 processes)
+ */
+enum class EErrorLogStage : uint8
+{
+	ELS_NoError,		// No error logs have been received/parsed yet
+	ELS_ErrorStart,		// The text indicating the start of an error log is being parsed
+	ELS_ErrorDesc,		// The text describing the error is being parsed
+	ELS_ErrorCallstack,	// The callstack for the error is being parsed
+	ELS_ErrorExit		// The post-error exit message is being parsed (error parsing is effectively complete)
+};
+
+/**
+ * Enum for specifying the suspend state of a process (typically the server)
+ */
+enum class ESuspendState : uint8
+{
+	Active,				// Process is currently active/not-suspended
+	Suspended,			// Process is currently suspended
+};
+
+
+// Struct definitions
+
+/**
+ * Struct used for storing and classifying each log error line
+ */
+struct FErrorLog
+{
+	/** The stage of this error log line */
+	EErrorLogStage Stage;
+
+	/** The error log line */
+	FString Line;
+
+
+	/**
+	 * Base constructor
+	 */
+	FErrorLog()
+		: Stage(EErrorLogStage::ELS_NoError)
+		, Line()
+	{
+	}
+
+	FErrorLog(EErrorLogStage InStage, const FString& InLine)
+		: Stage(InStage)
+		, Line(InLine)
+	{
+	}
+};
+
+/**
+ * Struct used for handling a launched UE4 client/server process
+ */
+struct FUnitTestProcess
+{
+	friend class UClientUnitTest;
+
+	/** Process handle for the launched process */
+	FProcHandle ProcessHandle;
+
+	/** The process ID */
+	uint32 ProcessID;
+
+	/** The suspend state of the process (implemented as a part of unit test code, does not relate to OS API) */
+	ESuspendState SuspendState;
+
+	/** Human-readable tag given to this process. Processes should be given an easily distinguishable tag, to identify them in code. */
+	FString ProcessTag;
+
+
+	/** Handle to StdOut for the launched process */
+	void* ReadPipe;
+
+	/** Handle to StdIn for the launched process (unused) */
+	void* WritePipe;
+
+
+	/** The base log type for this process (client? server? process?) */
+	ELogType BaseLogType;
+
+	/** The prefix to use for StdOut log output */
+	FString LogPrefix;
+
+	/** The OutputDeviceColor string, to use for setting the log output color */
+	const TCHAR* MainLogColor;
+
+	/** The log output color to use in the slate log window */
+	FSlateColor SlateLogColor;
+
+
+	/** If this process is outputting an error log, this is the current stage of error parsing (or ELS_NoError if not parsing) */
+	EErrorLogStage ErrorLogStage;
+
+	/** Gathered error log text */
+	TArray<FErrorLog> ErrorText;
+
+
+	/**
+	 * Base constructor
+	 */
+	FUnitTestProcess()
+		: ProcessHandle()
+		, ProcessID(0)
+		, SuspendState(ESuspendState::Active)
+		, ProcessTag(TEXT(""))
+		, ReadPipe(NULL)
+		, WritePipe(NULL)
+		, BaseLogType(ELogType::None)
+		, LogPrefix(TEXT(""))
+		, MainLogColor(COLOR_NONE)
+		, SlateLogColor(FSlateColor::UseForeground())
+		, ErrorLogStage(EErrorLogStage::ELS_NoError)
+		, ErrorText()
+	{
+	}
+};
+
+
+// Delegate definitions
+
+/**
+ * Delegate notifying that a process suspend state has changed
+ *
+ * @param NewSuspendState	The new process suspend state
+ */
+DECLARE_DELEGATE_OneParam(FOnSuspendStateChange, ESuspendState /*NewSuspendState*/);
+
+
+/**
+ * Base class for all unit tests which launch child processes, whether they be UE4 child processes, or other arbitrary programs.
+ *
+ * Handles management of child processes, memory usage tracking, log/stdout output gathering/printing, and crash detection.
+ */
+UCLASS()
+class NETCODEUNITTEST_API UProcessUnitTest : public UUnitTest
+{
+	GENERATED_UCLASS_BODY()
+
+	/** Runtime variables */
+protected:
+	/** Stores a reference to all running child processes tied to this unit test, for housekeeping */
+	TArray<TSharedPtr<FUnitTestProcess>> ActiveProcesses;
+
+	/** Last time there was a check for processes blocking progress */
+	double LastBlockingProcessCheck;
+
+public:
+	/** Delegate for notifying the UI, of a change in the unit test suspend state */
+	FOnSuspendStateChange OnSuspendStateChange;
+
+
+	/**
+	 * Interface for process unit tests
+	 */
+public:
+	/**
+	 * For implementation in subclasses, for helping to verify success/fail upon completion of unit tests
+	 * NOTE: Not called again once VerificationState is set
+	 * WARNING: Be careful when iterating InLogLines in multiple different for loops, if the sequence of detected logs is important
+	 *
+	 * @param InProcess		The process the log lines are from
+	 * @param InLogLines	The current log lines being received
+	 */
+	virtual void NotifyProcessLog(TWeakPtr<FUnitTestProcess> InProcess, const TArray<FString>& InLogLines)
+	{
+	}
+
+	/**
+	 * Notifies that there was a request to suspend/resume the unit test
+	 */
+	virtual void NotifySuspendRequest()
+	{
+	}
+
+	/**
+	 * Notifies when the suspend state of a process changes
+	 *
+	 * @param InProcess			The process whose suspend state has changed
+	 * @param InSuspendState	The new suspend state of the process
+	 */
+	virtual void NotifyProcessSuspendState(TWeakPtr<FUnitTestProcess> InProcess, ESuspendState InSuspendState);
+
+	/**
+	 * Notifies when a running process is detected as having finished/closed
+	 * NOTE: This will not get called, when ShutdownUnitTestProcess is used, only when a program ends by itself
+	 *
+	 * @param InProcess		The process which has finished
+	 */
+	virtual void NotifyProcessFinished(TWeakPtr<FUnitTestProcess> InProcess)
+	{
+	}
+
+
+protected:
+	/**
+	 * Starts a child process, tied to the unit test
+	 *
+	 * @param Path				The path to the process executable
+	 * @param InCommandline		The commandline that the child process should use
+	 * @param bMinimized		Starts the process with the window minimized
+	 * @return					Returns a pointer to the new processes handling struct
+	 */
+	virtual TWeakPtr<FUnitTestProcess> StartUnitTestProcess(FString Path, FString InCommandline, bool bMinimized=true);
+
+	/**
+	 * Starts a child UE4 process, tied to the unit test
+	 *
+	 * @param InCommandline		The commandline that the child process should use
+	 * @param bMinimized		Starts the process with the window minimized
+	 * @return					Returns a pointer to the new processes handling struct
+	 */
+	virtual TWeakPtr<FUnitTestProcess> StartUE4UnitTestProcess(FString InCommandline, bool bMinimized=true);
+
+	/**
+	 * Shuts-down/cleans-up a child process tied to the unit test
+	 *
+	 * @param InHandle	The handling struct for the process
+	 */
+	virtual void ShutdownUnitTestProcess(TSharedPtr<FUnitTestProcess> InHandle);
+
+	/**
+	 * If any errors logs were detected upon ShutdownUnitTestProcess, this is called to print them out
+	 *
+	 * @param InHandle	The handling struct for the process
+	 */
+	virtual void PrintUnitTestProcessErrors(TSharedPtr<FUnitTestProcess> InHandle);
+
+	/**
+	 * Processes the standard output (i.e. log output) for processes
+	 */
+	void PollProcessOutput();
+
+	/**
+	 * Updates (and if necessary, saves) the memory stats for processes
+	 */
+	void UpdateProcessStats();
+
+	/**
+	 * Whether or not a child process indicating a long/blocking task is running
+	 *
+	 * @param bLogIfFound	Outputs a log to the unit test window, if a blocking process is found
+	 * @return				Whether or not a blocking process was found
+	 */
+	bool IsBlockingProcessPresent(bool bLogIfFound=false);
+
+	/**
+	 * Checks incoming process logs, for any indication of a crash/error
+	 *
+	 * @param InProcess		The process the log output originates from
+	 * @param InLines		The incoming log lines
+	 */
+	virtual void CheckOutputForError(TSharedPtr<FUnitTestProcess> InProcess, const TArray<FString>& InLines);
+
+
+protected:
+	virtual void CleanupUnitTest() override;
+
+
+	virtual void UnitTick(float DeltaTime) override;
+
+	virtual void PostUnitTick(float DeltaTime) override;
+
+	virtual bool IsTickable() const override;
+
+
+	virtual void FinishDestroy() override;
+
+	virtual void ShutdownAfterError() override;
+};
diff --git a/Engine/Plugins/NetcodeUnitTest/NetcodeUnitTest/Source/NetcodeUnitTest/Classes/UnitTest.h b/Engine/Plugins/NetcodeUnitTest/NetcodeUnitTest/Source/NetcodeUnitTest/Classes/UnitTest.h
--- a/Engine/Plugins/NetcodeUnitTest/NetcodeUnitTest/Source/NetcodeUnitTest/Classes/UnitTest.h
+++ b/Engine/Plugins/NetcodeUnitTest/NetcodeUnitTest/Source/NetcodeUnitTest/Classes/UnitTest.h
@@ -307,8 +307,10 @@ protected:
 	virtual void ResetTimeout(FString ResetReason, bool bResetConnTimeout=false, uint32 MinDuration=0)
 	{
 		uint32 CurrentTimeout = FMath::Max(MinDuration, UnitTestTimeout);
+		double NewTimeoutExpire = FPlatformTime::Seconds() + (double)CurrentTimeout;
 
-		TimeoutExpire = FPlatformTime::Seconds() + (double)CurrentTimeout;
+		// Don't reset to a shorter timeout, than is already in place
+		TimeoutExpire = FMath::Max(NewTimeoutExpire, TimeoutExpire);
 
 		LastTimeoutReset = FPlatformTime::Seconds();
 		LastTimeoutResetEvent = ResetReason;
@@ -343,23 +345,51 @@ public:
 
 	/**
 	 * For implementation in subclasses, for helping to track local log entries related to this unit test.
-	 * This covers logs from within the unit test, and (for TCClientUnitTest) from processing net packets related to this unit test.
+	 * This covers logs from within the unit test, and (for UClientUnitTest) from processing net packets related to this unit test.
 	 *
 	 * NOTE: The parameters are the same as the unprocessed log 'serialize' parameters, to convert to a string, use:
 	 * FOutputDevice::FormatLogLine(Verbosity, Category, Data, GPrintLogTimes)
 	 *
 	 * NOTE: Verbosity ELogVerbosity::SetColor is a special category, whose log messages can be ignored.
 	 *
-	 * @param LogType	The type of local log message this is
+	 * @param InLogType	The type of local log message this is
 	 * @param Data		The base log message being written
 	 * @param Verbosity	The warning/filter level for the log message
 	 * @param Category	The log message category (LogNet, LogNetTraffic etc.)
 	 */
-	virtual void NotifyLocalLog(ELogType LogType, const TCHAR* Data, ELogVerbosity::Type Verbosity, const class FName& Category);
+	virtual void NotifyLocalLog(ELogType InLogType, const TCHAR* Data, ELogVerbosity::Type Verbosity, const class FName& Category);
+
+
+	/**
+	 * Notifies that there was a request to enable/disable developer mode
+	 *
+	 * @param bInDeveloperMode	Whether or not developer mode is being enabled/disabled
+	 */
+	void NotifyDeveloperModeRequest(bool bInDeveloperMode);
+
+	/**
+	 * Notifies that there was a request to execute a console command for the unit test, which can occur in a specific context,
+	 * e.g. for a unit test server, for a local minimal-client (within the unit test), or for a separate unit test client process
+	 *
+	 * @param CommandContext	The context (local/server/client?) for the console command
+	 * @param Command			The command to be executed
+	 * @return					Whether or not the command was handled
+	 */
+	virtual bool NotifyConsoleCommandRequest(FString CommandContext, FString Command);
+
+	/**
+	 * Outputs the list of console command contexts, that this unit test supports (which can include custom contexts in subclasses)
+	 *
+	 * @param OutList				Outputs the list of supported console command contexts
+	 * @param OutDefaultContext		Outputs the context which should be auto-selected/defaulted-to
+	 */
+	virtual void GetCommandContextList(TArray<TSharedPtr<FString>>& OutList, FString& OutDefaultContext);
 
 
 	virtual void PostUnitTick(float DeltaTime) override;
 
+	virtual bool IsTickable() const override;
+
 	virtual void TickIsComplete(float DeltaTime) override;
 
 
diff --git a/Engine/Plugins/NetcodeUnitTest/NetcodeUnitTest/Source/NetcodeUnitTest/Classes/UnitTestBase.h b/Engine/Plugins/NetcodeUnitTest/NetcodeUnitTest/Source/NetcodeUnitTest/Classes/UnitTestBase.h
--- a/Engine/Plugins/NetcodeUnitTest/NetcodeUnitTest/Source/NetcodeUnitTest/Classes/UnitTestBase.h
+++ b/Engine/Plugins/NetcodeUnitTest/NetcodeUnitTest/Source/NetcodeUnitTest/Classes/UnitTestBase.h
@@ -64,6 +64,6 @@ public:
 	// Must override in subclasses, that need ticking
 	virtual bool IsTickable() const
 	{
-		return true;
+		return false;
 	}
 };
diff --git a/Engine/Plugins/NetcodeUnitTest/NetcodeUnitTest/Source/NetcodeUnitTest/Private/ClientUnitTest.cpp b/Engine/Plugins/NetcodeUnitTest/NetcodeUnitTest/Source/NetcodeUnitTest/Private/ClientUnitTest.cpp
--- a/Engine/Plugins/NetcodeUnitTest/NetcodeUnitTest/Source/NetcodeUnitTest/Private/ClientUnitTest.cpp
+++ b/Engine/Plugins/NetcodeUnitTest/NetcodeUnitTest/Source/NetcodeUnitTest/Private/ClientUnitTest.cpp
@@ -44,11 +44,14 @@ UClientUnitTest::UClientUnitTest(const FObjectInitializer& ObjectInitializer)
 	, BaseClientParameters(TEXT(""))
 	, AllowedClientActors()
 	, AllowedClientRPCs()
-	, ActiveProcesses()
 	, ServerHandle(NULL)
 	, ServerAddress(TEXT(""))
 	, BeaconAddress(TEXT(""))
 	, ClientHandle(NULL)
+	, bBlockingServerDelay(false)
+	, bBlockingClientDelay(false)
+	, bBlockingFakeClientDelay(false)
+	, NextBlockingTimeout(0.0)
 	, UnitWorld(NULL)
 	, UnitNotify(NULL)
 	, UnitNetDriver(NULL)
@@ -60,7 +63,6 @@ UClientUnitTest::UClientUnitTest(const FObjectInitializer& ObjectInitializer)
 	, bReceivedPong(false)
 	, ControlBunchSequence(0)
 	, PendingNetActorChans()
-	, OnServerSuspendState()
 {
 }
 
@@ -468,12 +470,26 @@ void UClientUnitTest::NotifyProcessLog(TWeakPtr<FUnitTestProcess> InProcess, con
 				// Fire off fake client connection
 				if (UnitConn == NULL)
 				{
-					FString LogMsg = FString::Printf(TEXT("Detected successful server startup, launching fake client."));
+					bool bBlockingProcess = IsBlockingProcessPresent(true);
 
-					UNIT_LOG(ELogType::StatusImportant, TEXT("%s"), *LogMsg);
-					UNIT_STATUS_LOG(ELogType::StatusVerbose, TEXT("%s"), *LogMsg);
+					if (bBlockingProcess)
+					{
+						FString LogMsg = TEXT("Detected successful server startup, delaying fake client due to blocking process.");
+
+						UNIT_LOG(ELogType::StatusImportant, TEXT("%s"), *LogMsg);
+						UNIT_STATUS_LOG(ELogType::StatusVerbose, TEXT("%s"), *LogMsg);
+
+						bBlockingFakeClientDelay = true;
+					}
+					else
+					{
+						FString LogMsg = TEXT("Detected successful server startup, launching fake client.");
 
-					ConnectFakeClient();
+						UNIT_LOG(ELogType::StatusImportant, TEXT("%s"), *LogMsg);
+						UNIT_STATUS_LOG(ELogType::StatusVerbose, TEXT("%s"), *LogMsg);
+
+						ConnectFakeClient();
+					}
 				}
 
 				ResetTimeout(FString(TEXT("ServerReady: ")) + MatchedLine);
@@ -488,7 +504,7 @@ void UClientUnitTest::NotifyProcessLog(TWeakPtr<FUnitTestProcess> InProcess, con
 		{
 			if (ServerTimeoutResetLogs->ContainsByPredicate(SearchInLogLine))
 			{
-				ResetTimeout(FString(TEXT("ServerTimeoutReset: ")) + MatchedLine, true);
+				ResetTimeout(FString(TEXT("ServerTimeoutReset: ")) + MatchedLine, true, 60);
 			}
 		}
 	}
@@ -499,7 +515,7 @@ void UClientUnitTest::NotifyProcessLog(TWeakPtr<FUnitTestProcess> InProcess, con
 		{
 			if (ClientTimeoutResetLogs->ContainsByPredicate(SearchInLogLine))
 			{
-				ResetTimeout(FString(TEXT("ClientTimeoutReset: ")) + MatchedLine, true);
+				ResetTimeout(FString(TEXT("ClientTimeoutReset: ")) + MatchedLine, true, 60);
 			}
 		}
 	}
@@ -508,8 +524,80 @@ void UClientUnitTest::NotifyProcessLog(TWeakPtr<FUnitTestProcess> InProcess, con
 	//				to reset their connection timeout upon server progress, if they fully startup before the server does
 }
 
+void UClientUnitTest::NotifyProcessFinished(TWeakPtr<FUnitTestProcess> InProcess)
+{
+	Super::NotifyProcessFinished(InProcess);
+
+	if (InProcess.IsValid())
+	{
+		bool bServerFinished = false;
+		bool bClientFinished  = false;
+
+		if (ServerHandle.IsValid() && ServerHandle.HasSameObject(InProcess.Pin().Get()))
+		{
+			bServerFinished = true;
+		}
+		else if (ClientHandle.IsValid() && ClientHandle.HasSameObject(InProcess.Pin().Get()))
+		{
+			bClientFinished = true;
+		}
+
+		if (bServerFinished || bClientFinished)
+		{
+			bool bProcessError = false;
+			FString UpdateMsg;
+
+			// If the server just finished, cleanup the fake client
+			if (bServerFinished)
+			{
+				FString LogMsg = TEXT("Server process has finished, cleaning up fake client.");
+
+				UNIT_LOG(ELogType::StatusImportant, TEXT("%s"), *LogMsg);
+				UNIT_STATUS_LOG(ELogType::StatusVerbose, TEXT("%s"), *LogMsg);
+
+				if (UnitConn != NULL)
+				{
+					UnitConn->Close();
+				}
+
+				// Immediately cleanup the fake client (don't wait for end-of-life cleanup in CleanupUnitTest)
+				CleanupFakeClient();
+
+
+				// If a server exit was unexpected, mark the unit test as broken
+				if (!(UnitTestFlags & EUnitTestFlags::IgnoreServerCrash) && VerificationState == EUnitTestVerification::Unverified)
+				{
+					UpdateMsg = TEXT("Unexpected server exit, marking unit test as needing update.");
+					bProcessError = true;
+				}
+			}
+
+			// If a client exit was unexpected, mark the unit test as broken
+			if (bClientFinished && !(UnitTestFlags & EUnitTestFlags::IgnoreClientCrash) &&
+				VerificationState == EUnitTestVerification::Unverified)
+			{
+				UpdateMsg = TEXT("Unexpected client exit, marking unit test as needing update.");
+				bProcessError = true;
+			}
+
+
+			// If either the client/server finished, process the error
+			if (bProcessError)
+			{
+				UNIT_LOG(ELogType::StatusFailure | ELogType::StyleBold, TEXT("%s"), *UpdateMsg);
+				UNIT_STATUS_LOG(ELogType::StatusFailure | ELogType::StatusVerbose | ELogType::StyleBold, TEXT("%s"), *UpdateMsg);
+
+				VerificationState = EUnitTestVerification::VerifiedNeedsUpdate;
+			}
+		}
+	}
+}
+
 void UClientUnitTest::NotifySuspendRequest()
 {
+	// @todo #JohnBFeature: Currently on suspends the server, suspend the client as well, and also add more granularity
+	//						(deciding which to suspend)
+
 	TSharedPtr<FUnitTestProcess> CurProcess = (ServerHandle.IsValid() ? ServerHandle.Pin() : NULL);
 
 	if (CurProcess.IsValid())
@@ -585,93 +673,68 @@ void UClientUnitTest::NotifySuspendRequest()
 
 void UClientUnitTest::NotifyProcessSuspendState(TWeakPtr<FUnitTestProcess> InProcess, ESuspendState InSuspendState)
 {
-	InProcess.Pin()->SuspendState = InSuspendState;
+	Super::NotifyProcessSuspendState(InProcess, InSuspendState);
 
 	if (InProcess == ServerHandle)
 	{
-		OnServerSuspendState.ExecuteIfBound(InSuspendState);
+		OnSuspendStateChange.ExecuteIfBound(InSuspendState);
 	}
 
 	// @todo #JohnBLowPri: Want to support client process debugging in future?
 }
 
 
-void UClientUnitTest::NotifyDeveloperModeRequest(bool bInDeveloperMode)
-{
-	bDeveloperMode = bInDeveloperMode;
-}
-
-
 bool UClientUnitTest::NotifyConsoleCommandRequest(FString CommandContext, FString Command)
 {
-	bool bHandled = false;
+	bool bHandled = Super::NotifyConsoleCommandRequest(CommandContext, Command);
 
-	if (CommandContext == TEXT("Local") || CommandContext == TEXT("Global"))
+	if (!bHandled)
 	{
-		UWorld* TargetWorld = (CommandContext == TEXT("Local") ? UnitWorld : NULL);
-
-		// Don't execute commands that crash
-		TArray<FString> BadCmds;
-		BadCmds.Add(TEXT("exit"));
-
-		if (!BadCmds.Contains(Command))
+		if (CommandContext == TEXT("Local"))
 		{
-			// @todo #JohnB: Should this mark the log origin, as from the unit test?
-			// @todo #JohnB: In general, I'm not sure how I handle the log-origin of UI-triggered events;
-			//				they maybe should be captured/categorized better
 			UNIT_LOG_BEGIN(this, ELogType::OriginConsole);
-			bHandled = GEngine->Exec(TargetWorld, *Command, *GLog);
+			bHandled = GEngine->Exec(UnitWorld, *Command, *GLog);
 			UNIT_LOG_END();
 		}
-		else
+		else if (CommandContext == TEXT("Server"))
 		{
-			UNIT_LOG(ELogType::OriginConsole,
-						TEXT("Can't execute command '%s', it's in the 'bad commands' list (i.e. probably crashes)"), *Command);
-		}
-	}
-	else if (CommandContext == TEXT("Server"))
-	{
-		// @todo #JohnBBug: Perhaps add extra checks here, to be sure we're ready to send console commands?
-		//
-		//				UPDATE: Yes, this is a good idea, because if the client hasn't gotten to the correct login stage
-		//				(NMT_Join or such, need to check when server rejects non-login control commands),
-		//				then it leads to an early disconnect when you try to spam-send a command early.
-		//
-		//				It's easy to test this, just type in a command before join, and hold down enter on the edit box to spam it.
-		if (UnitConn != NULL)
-		{
-			FOutBunch* ControlChanBunch = NUTNet::CreateChannelBunch(ControlBunchSequence, UnitConn, CHTYPE_Control, 0);
+			// @todo #JohnBBug: Perhaps add extra checks here, to be sure we're ready to send console commands?
+			//
+			//				UPDATE: Yes, this is a good idea, because if the client hasn't gotten to the correct login stage
+			//				(NMT_Join or such, need to check when server rejects non-login control commands),
+			//				then it leads to an early disconnect when you try to spam-send a command early.
+			//
+			//				It's easy to test this, just type in a command before join, and hold down enter on the edit box to spam it.
+			if (UnitConn != NULL)
+			{
+				FOutBunch* ControlChanBunch = NUTNet::CreateChannelBunch(ControlBunchSequence, UnitConn, CHTYPE_Control, 0);
 
-			uint8 ControlMsg = NMT_NUTControl;
-			uint8 ControlCmd = ENUTControlCommand::Command_NoResult;
-			FString Cmd = Command;
+				uint8 ControlMsg = NMT_NUTControl;
+				uint8 ControlCmd = ENUTControlCommand::Command_NoResult;
+				FString Cmd = Command;
 
-			*ControlChanBunch << ControlMsg;
-			*ControlChanBunch << ControlCmd;
-			*ControlChanBunch << Cmd;
+				*ControlChanBunch << ControlMsg;
+				*ControlChanBunch << ControlCmd;
+				*ControlChanBunch << Cmd;
 
-			NUTNet::SendControlBunch(UnitConn, *ControlChanBunch);
+				NUTNet::SendControlBunch(UnitConn, *ControlChanBunch);
 
 
-			UNIT_LOG(ELogType::OriginConsole, TEXT("Sent command '%s' to server."), *Command);
+				UNIT_LOG(ELogType::OriginConsole, TEXT("Sent command '%s' to server."), *Command);
 
-			bHandled = true;
+				bHandled = true;
+			}
+			else
+			{
+				UNIT_LOG(ELogType::OriginConsole, TEXT("Failed to send console command '%s', no server connection."), *Command);
+			}
 		}
-		else
+		else if (CommandContext == TEXT("Client"))
 		{
-			UNIT_LOG(ELogType::OriginConsole, TEXT("Failed to send console command '%s', no server connection."), *Command);
-		}
-	}
-	else if (CommandContext == TEXT("Client"))
-	{
-		// @todo #JohnBFeature
+			// @todo #JohnBFeature
 
-		UNIT_LOG(ELogType::OriginConsole, TEXT("Client console commands not yet implemented"));
-	}
-	else
-	{
-		UNIT_LOG(ELogType::StatusFailure | ELogType::OriginConsole,
-					TEXT("Unknown console command context '%s' - this is a bug that should be fixed."), *CommandContext);
+			UNIT_LOG(ELogType::OriginConsole, TEXT("Client console commands not yet implemented"));
+		}
 	}
 
 	return bHandled;
@@ -679,7 +742,8 @@ bool UClientUnitTest::NotifyConsoleCommandRequest(FString CommandContext, FStrin
 
 void UClientUnitTest::GetCommandContextList(TArray<TSharedPtr<FString>>& OutList, FString& OutDefaultContext)
 {
-	OutList.Add(MakeShareable(new FString(TEXT("Global"))));
+	Super::GetCommandContextList(OutList, OutDefaultContext);
+
 	OutList.Add(MakeShareable(new FString(TEXT("Local"))));
 
 	if (!!(UnitTestFlags & EUnitTestFlags::LaunchServer))
@@ -988,18 +1052,20 @@ void UClientUnitTest::ResetTimeout(FString ResetReason, bool bResetConnTimeout/*
 
 	if (bResetConnTimeout)
 	{
-		ResetConnTimeout();
+		ResetConnTimeout((float)(FMath::Max(MinDuration, UnitTestTimeout)));
 	}
 }
 
-void UClientUnitTest::ResetConnTimeout()
+void UClientUnitTest::ResetConnTimeout(float Duration)
 {
 	if (UnitConn != NULL && UnitConn->State != USOCK_Closed && UnitConn->Driver != NULL)
 	{
 		// @todo #JohnBHack: This is a slightly hacky way of setting the timeout to a large value, which will be overridden by newly
 		//				received packets, making it unsuitable for most situations (except crashes - but that could still be subject
 		//				to a race condition)
-		UnitConn->LastReceiveTime = UnitConn->Driver->Time + (float)(TimeoutExpire - FPlatformTime::Seconds());
+		double NewLastReceiveTime = UnitConn->Driver->Time + Duration;
+
+		UnitConn->LastReceiveTime = FMath::Max(NewLastReceiveTime, UnitConn->LastReceiveTime);
 	}
 }
 
@@ -1017,12 +1083,38 @@ bool UClientUnitTest::ExecuteUnitTest()
 	{
 		if (!!(UnitTestFlags & EUnitTestFlags::LaunchServer))
 		{
-			StartUnitTestServer();
+			bool bBlockingProcess = IsBlockingProcessPresent(true);
+
+			if (bBlockingProcess)
+			{
+				FString LogMsg = TEXT("Delaying server startup due to blocking process");
+
+				UNIT_LOG(ELogType::StatusImportant, TEXT("%s"), *LogMsg);
+				UNIT_STATUS_LOG(ELogType::StatusVerbose, TEXT("%s"), *LogMsg);
+
+				bBlockingServerDelay = true;
+			}
+			else
+			{
+				StartUnitTestServer();
+			}
 
 			if (!!(UnitTestFlags & EUnitTestFlags::LaunchClient))
 			{
-				// Client handle is set outside of StartUnitTestClient, in case support for multiple clients is added later
-				ClientHandle = StartUnitTestClient(ServerAddress);
+				if (bBlockingProcess)
+				{
+					FString LogMsg = TEXT("Delaying client startup due to blocking process");
+
+					UNIT_LOG(ELogType::StatusImportant, TEXT("%s"), *LogMsg);
+					UNIT_STATUS_LOG(ELogType::StatusVerbose, TEXT("%s"), *LogMsg);
+
+					bBlockingClientDelay = true;
+				}
+				else
+				{
+					// Client handle is set outside of StartUnitTestClient, in case support for multiple clients is added later
+					ClientHandle = StartUnitTestClient(ServerAddress);
+				}
 			}
 		}
 
@@ -1055,19 +1147,11 @@ void UClientUnitTest::CleanupUnitTest()
 {
 	CleanupFakeClient();
 
-	for (int32 i=ActiveProcesses.Num()-1; i>=0; i--)
-	{
-		if (ActiveProcesses[i].IsValid())
-		{
-			ShutdownUnitTestProcess(ActiveProcesses[i]);
-		}
-	}
-
 #if !UE_BUILD_SHIPPING
 	RemoveProcessEventCallback(this, &UClientUnitTest::InternalScriptProcessEvent);
 #endif
 
-	UUnitTest::CleanupUnitTest();
+	Super::CleanupUnitTest();
 }
 
 bool UClientUnitTest::ConnectFakeClient(FUniqueNetIdRepl* InNetID/*=NULL*/)
@@ -1149,8 +1233,7 @@ bool UClientUnitTest::ConnectFakeClient(FUniqueNetIdRepl* InNetID/*=NULL*/)
 
 					CurUnitConn->ReplicatedActorSpawnDel.BindUObject(this, &UClientUnitTest::NotifyAllowNetActor);
 
-					// If you don't have to wait for any requirements, execute immediately
-					if (!(UnitTestFlags & (EUnitTestFlags::RequirementsMask)))
+					if (HasAllRequirements())
 					{
 						ResetTimeout(TEXT("ExecuteClientUnitTest (ExecuteUnitTest)"));
 						ExecuteClientUnitTest();
@@ -1270,7 +1353,7 @@ void UClientUnitTest::StartUnitTestServer()
 		// Increment the server port used by 10, for every unit test
 		static int ServerPortOffset = 0;
 		int ServerPort = DefaultPort + 50 + (++ServerPortOffset * 10);
-		int BeaconPort = ServerPort + 5;
+		int ServerBeaconPort = ServerPort + 5;
 
 
 		// Setup the launch URL
@@ -1278,10 +1361,10 @@ void UClientUnitTest::StartUnitTestServer()
 
 		if (!!(UnitTestFlags & EUnitTestFlags::BeaconConnect))
 		{
-			ServerParameters += FString::Printf(TEXT(" -BeaconPort=%i"), BeaconPort);
+			ServerParameters += FString::Printf(TEXT(" -BeaconPort=%i"), ServerBeaconPort);
 		}
 
-		ServerHandle = StartUnitTestProcess(ServerParameters);
+		ServerHandle = StartUE4UnitTestProcess(ServerParameters);
 
 		if (ServerHandle.IsValid())
 		{
@@ -1289,12 +1372,12 @@ void UClientUnitTest::StartUnitTestServer()
 
 			if (!!(UnitTestFlags & EUnitTestFlags::BeaconConnect))
 			{
-				BeaconAddress = FString::Printf(TEXT("127.0.0.1:%i"), BeaconPort);
+				BeaconAddress = FString::Printf(TEXT("127.0.0.1:%i"), ServerBeaconPort);
 			}
 
 			auto CurHandle = ServerHandle.Pin();
 
-			CurHandle->ProcessTag = TEXT("Server");
+			CurHandle->ProcessTag = FString::Printf(TEXT("UE4_Server_%i"), CurHandle->ProcessID);
 			CurHandle->BaseLogType = ELogType::Server;
 			CurHandle->LogPrefix = TEXT("[SERVER]");
 			CurHandle->MainLogColor = COLOR_CYAN;
@@ -1313,7 +1396,10 @@ FString UClientUnitTest::ConstructServerParameters()
 	// NOTE: Without '-CrashForUAT'/'-unattended' the auto-reporter can pop up
 	FString Parameters = FString(FApp::GetGameName()) + TEXT(" ") + BaseServerURL + TEXT(" -server ") + BaseServerParameters +
 							TEXT(" -Log=UnitTestServer.log -forcelogflush -stdout -AllowStdOutLogVerbosity -ddc=noshared") +
-							TEXT(" -unattended -CrashForUAT -NoShaderWorker");
+							TEXT(" -unattended -CrashForUAT");
+	
+							// Removed this, to support detection of shader compilation, based on shader compiler .exe
+							//TEXT(" -NoShaderWorker");
 
 	return Parameters;
 }
@@ -1329,14 +1415,14 @@ TWeakPtr<FUnitTestProcess> UClientUnitTest::StartUnitTestClient(FString ConnectI
 
 	FString ClientParameters = ConstructClientParameters(ConnectIP);
 
-	ReturnVal = StartUnitTestProcess(ClientParameters, bMinimized);
+	ReturnVal = StartUE4UnitTestProcess(ClientParameters, bMinimized);
 
 	if (ReturnVal.IsValid())
 	{
 		auto CurHandle = ReturnVal.Pin();
 
 		// @todo #JohnBMultiClient: If you add support for multiple clients, make the log prefix numbered, also try to differentiate colours
-		CurHandle->ProcessTag = TEXT("Client");
+		CurHandle->ProcessTag = FString::Printf(TEXT("UE4_Client_%i"), CurHandle->ProcessID);
 		CurHandle->BaseLogType = ELogType::Client;
 		CurHandle->LogPrefix = TEXT("[CLIENT]");
 		CurHandle->MainLogColor = COLOR_GREEN;
@@ -1352,254 +1438,30 @@ FString UClientUnitTest::ConstructClientParameters(FString ConnectIP)
 	// NOTE: Without '-CrashForUAT'/'-unattended' the auto-reporter can pop up
 	FString Parameters = FString(FApp::GetGameName()) + TEXT(" ") + ConnectIP + BaseClientURL + TEXT(" -game ") + BaseClientParameters +
 							TEXT(" -Log=UnitTestClient.log -forcelogflush -stdout -AllowStdOutLogVerbosity -ddc=noshared -nosplash") +
-							TEXT(" -unattended -CrashForUAT -NoShaderWorker -nosound");
-
-	return Parameters;
-}
-
-TWeakPtr<FUnitTestProcess> UClientUnitTest::StartUnitTestProcess(FString InCommandline, bool bMinimized/*=true*/)
-{
-	TSharedPtr<FUnitTestProcess> ReturnVal = MakeShareable(new FUnitTestProcess());
-
-	FString GamePath = FPlatformProcess::GenerateApplicationPath(FApp::GetName(), FApp::GetBuildConfiguration());
-
-	verify(FPlatformProcess::CreatePipe(ReturnVal->ReadPipe, ReturnVal->WritePipe));
-
-	UNIT_LOG(ELogType::StatusImportant, TEXT("Starting process with parameters: %s"), *InCommandline);
-
-	ReturnVal->ProcessHandle = FPlatformProcess::CreateProc(*GamePath, *InCommandline, true, bMinimized, false, &ReturnVal->ProcessID,
-															0, NULL, ReturnVal->WritePipe);
-
-	if (ReturnVal->ProcessHandle.IsValid())
-	{
-		ReturnVal->ProcessTag = FString::Printf(TEXT("Process_%i"), ReturnVal->ProcessID);
-
-		ActiveProcesses.Add(ReturnVal);
-	}
-	else
-	{
-		UNIT_LOG(ELogType::StatusFailure, TEXT("Failed to start process"));
-	}
-
-	return ReturnVal;
-}
-
-void UClientUnitTest::ShutdownUnitTestProcess(TSharedPtr<FUnitTestProcess> InHandle)
-{
-	if (InHandle->ProcessHandle.IsValid())
-	{
-		FString LogMsg = FString::Printf(TEXT("Shutting down process '%s'."), *InHandle->ProcessTag);
-
-		UNIT_LOG(ELogType::StatusImportant, TEXT("%s"), *LogMsg);
-		UNIT_STATUS_LOG(ELogType::StatusVerbose, TEXT("%s"), *LogMsg);
-
-
-		// @todo #JohnBHighPri: Restore 'true' here, once the issue where killing child processes sometimes kills all processes, is fixed
-		FPlatformProcess::TerminateProc(InHandle->ProcessHandle);//, true);
-
-#if TARGET_UE4_CL < CL_CLOSEPROC
-		InHandle->ProcessHandle.Close();
-#else
-		FPlatformProcess::CloseProc(InHandle->ProcessHandle);
-#endif
-	}
-
-	FPlatformProcess::ClosePipe(InHandle->ReadPipe, InHandle->WritePipe);
-
-
-	// Print out any detected error logs
-	if (InHandle->ErrorLogStage != EErrorLogStage::ELS_NoError && InHandle->ErrorText.Num() > 0)
-	{
-		FString LogMsg = FString::Printf(TEXT("Detected a crash in process '%s':"),
-											*InHandle->ProcessTag);
+							TEXT(" -unattended -CrashForUAT -nosound");
 
-		UNIT_LOG(ELogType::StatusImportant, TEXT("%s"), *LogMsg);
-		UNIT_STATUS_LOG(ELogType::StatusImportant, TEXT("%s (click 'Advanced Summary' for more detail)"), *LogMsg);
-
-
-		// If this was the server, and we were not expecting a crash, print out a warning
-		if (!(UnitTestFlags & EUnitTestFlags::ExpectServerCrash) && InHandle == ServerHandle.Pin())
-		{
-			LogMsg = TEXT("WARNING: Got server crash, but unit test not marked as expecting a server crash.");
-
-			STATUS_SET_COLOR(FLinearColor(1.0, 1.0, 0.0));
-
-			UNIT_LOG(ELogType::StatusWarning, TEXT("%s"), *LogMsg);
-			UNIT_STATUS_LOG(ELogType::StatusWarning, TEXT("%s"), *LogMsg);
-
-			STATUS_RESET_COLOR();
-		}
+							// Removed this, to support detection of shader compilation, based on shader compiler .exe
+							//TEXT(" -NoShaderWorker")
 
-
-		uint32 CallstackCount = 0;
-
-		for (auto CurErrorLine : InHandle->ErrorText)
-		{
-			ELogType CurLogType = ELogType::None;
-
-			if (CurErrorLine.Stage == EErrorLogStage::ELS_ErrorStart)
-			{
-				CurLogType = ELogType::StatusAdvanced;
-			}
-			else if (CurErrorLine.Stage == EErrorLogStage::ELS_ErrorDesc)
-			{
-				CurLogType = ELogType::StatusImportant | ELogType::StatusAdvanced | ELogType::StyleBold;
-			}
-			else if (CurErrorLine.Stage == EErrorLogStage::ELS_ErrorCallstack)
-			{
-				CurLogType = ELogType::StatusAdvanced;
-
-				// Include the first five callstack lines in the main summary printout
-				if (CallstackCount < 5)
-				{
-					CurLogType |= ELogType::StatusImportant;
-				}
-
-				CallstackCount++;
-			}
-			else if (CurErrorLine.Stage == EErrorLogStage::ELS_ErrorExit)
-			{
-				continue;
-			}
-			else
-			{
-				// Make sure there's a check for all error log stages
-				UNIT_ASSERT(false);
-			}
-
-			// For the current unit-test log window, all entries are marked 'StatusImportant'
-			UNIT_LOG(CurLogType | ELogType::StatusImportant, TEXT("%s"), *CurErrorLine.Line);
-
-			UNIT_STATUS_LOG(CurLogType, TEXT("%s"), *CurErrorLine.Line);
-		}
-	}
-
-
-	ActiveProcesses.Remove(InHandle);
+	return Parameters;
 }
 
-
-void UClientUnitTest::UpdateProcessStats()
+void UClientUnitTest::PrintUnitTestProcessErrors(TSharedPtr<FUnitTestProcess> InHandle)
 {
-	SIZE_T TotalProcessMemoryUsage = 0;
-
-	for (auto CurHandle : ActiveProcesses)
-	{
-		// Check unit test memory stats, and update if necessary (NOTE: Processes not guaranteed to still be running)
-		if (CurHandle.IsValid() && CurHandle->ProcessID != 0)
-		{
-			SIZE_T ProcessMemUsage = 0;
-
-			if (FPlatformProcess::GetApplicationMemoryUsage(CurHandle->ProcessID, &ProcessMemUsage) && ProcessMemUsage != 0)
-			{
-				TotalProcessMemoryUsage += ProcessMemUsage;
-			}
-		}
-	}
-
-	if (TotalProcessMemoryUsage > 0)
+	// If this was the server, and we were not expecting a crash, print out a warning
+	if (!(UnitTestFlags & EUnitTestFlags::ExpectServerCrash) && ServerHandle.IsValid() && InHandle == ServerHandle.Pin())
 	{
-		CurrentMemoryUsage = TotalProcessMemoryUsage;
+		FString LogMsg = TEXT("WARNING: Got server crash, but unit test not marked as expecting a server crash.");
 
-		float RunningTime = (float)(FPlatformTime::Seconds() - StartTime);
+		STATUS_SET_COLOR(FLinearColor(1.0, 1.0, 0.0));
 
-		// Update saved memory usage stats, to help with tracking estimated memory usage, on future unit test launches
-		if (CurrentMemoryUsage > PeakMemoryUsage)
-		{
-			if (PeakMemoryUsage == 0)
-			{
-				bFirstTimeStats = true;
-			}
-
-			PeakMemoryUsage = CurrentMemoryUsage;
-
-			// Reset TimeToPeakMem
-			TimeToPeakMem = RunningTime;
+		UNIT_LOG(ELogType::StatusWarning, TEXT("%s"), *LogMsg);
+		UNIT_STATUS_LOG(ELogType::StatusWarning, TEXT("%s"), *LogMsg);
 
-			SaveConfig();
-		}
-		// Check if we have hit a new record, for the time it takes to get within 90% of PeakMemoryUsage
-		else if (RunningTime < TimeToPeakMem && (CurrentMemoryUsage * 100) >= (PeakMemoryUsage * 90))
-		{
-			TimeToPeakMem = RunningTime;
-			SaveConfig();
-		}
+		STATUS_RESET_COLOR();
 	}
-}
-
-void UClientUnitTest::CheckOutputForError(TSharedPtr<FUnitTestProcess> InProcess, const TArray<FString>& InLines)
-{
-	// Array of log messages that can indicate the start of an error
-	static const TArray<FString> ErrorStartLogs = TArrayBuilder<FString>()
-		.Add(FString(TEXT("Windows GetLastError: ")))
-		.Add(FString(TEXT("=== Critical error: ===")));
-
-
-	for (FString CurLine : InLines)
-	{
-		// Using 'ContainsByPredicate' as an iterator
-		const auto CheckForErrorLog =
-			[&CurLine](const FString& ErrorLine)
-			{
-				return CurLine.Contains(ErrorLine);
-			};
-
-		// Search for the beginning of an error log message
-		if (InProcess->ErrorLogStage == EErrorLogStage::ELS_NoError)
-		{
-			if (ErrorStartLogs.ContainsByPredicate(CheckForErrorLog))
-			{
-				InProcess->ErrorLogStage = EErrorLogStage::ELS_ErrorStart;
-
-				// Reset the timeout for both the unit test and unit test connection here, as callstack logs are prone to failure
-				//	(do it for at least two minutes as well, as it can take a long time to get a crashlog)
-				ResetTimeout(TEXT("Detected crash."), true, 120);
-			}
-		}
 
-
-		if (InProcess->ErrorLogStage != EErrorLogStage::ELS_NoError)
-		{
-			// Regex pattern for matching callstack logs - matches:
-			//	" (0x000007fefe22cacd) + 0 bytes"
-			//	" {0x000007fefe22cacd} + 0 bytes"
-			const FRegexPattern CallstackPattern(TEXT("\\s[\\(|\\{]0x[0-9,a-f]+[\\)|\\}] \\+ [0-9]+ bytes"));
-
-			// Matches:
-			//	"ntldll.dll"
-			const FRegexPattern AltCallstackPattern(TEXT("^[a-z,A-Z,0-9,\\-,_]+\\.[exe|dll]"));
-
-
-			// Check for the beginning of description logs
-			if (InProcess->ErrorLogStage == EErrorLogStage::ELS_ErrorStart &&
-				!ErrorStartLogs.ContainsByPredicate(CheckForErrorLog))
-			{
-				InProcess->ErrorLogStage = EErrorLogStage::ELS_ErrorDesc;
-			}
-
-			// Check-for/verify callstack logs
-			if (InProcess->ErrorLogStage == EErrorLogStage::ELS_ErrorDesc ||
-				InProcess->ErrorLogStage == EErrorLogStage::ELS_ErrorCallstack)
-			{
-				FRegexMatcher CallstackMatcher(CallstackPattern, CurLine);
-				FRegexMatcher AltCallstackMatcher(AltCallstackPattern, CurLine);
-
-				if (CallstackMatcher.FindNext() || AltCallstackMatcher.FindNext())
-				{
-					InProcess->ErrorLogStage = EErrorLogStage::ELS_ErrorCallstack;
-				}
-				else if (InProcess->ErrorLogStage == EErrorLogStage::ELS_ErrorCallstack)
-				{
-					InProcess->ErrorLogStage = EErrorLogStage::ELS_ErrorExit;
-				}
-			}
-
-			// The rest of the lines, after callstack parsing, should be ELS_ErrorExit logs - these are not verified though (no need)
-
-
-			InProcess->ErrorText.Add(FErrorLog(InProcess->ErrorLogStage, CurLine));
-		}
-	}
+	Super::PrintUnitTestProcessErrors(InHandle);
 }
 
 
@@ -1674,283 +1536,76 @@ void UClientUnitTest::InternalNotifyNetworkFailure(UWorld* InWorld, UNetDriver*
 
 void UClientUnitTest::UnitTick(float DeltaTime)
 {
-	Super::UnitTick(DeltaTime);
-
-	if (!!(UnitTestFlags & EUnitTestFlags::NotifyNetActors) && PendingNetActorChans.Num() > 0 && UnitConn != NULL)
+	if (bBlockingServerDelay || bBlockingClientDelay || bBlockingFakeClientDelay)
 	{
-		for (int32 i=PendingNetActorChans.Num()-1; i>=0; i--)
-		{
-			UActorChannel* CurChan = Cast<UActorChannel>(UnitConn->Channels[PendingNetActorChans[i]]);
-
-			if (CurChan != NULL && CurChan->Actor != NULL)
-			{
-				NotifyNetActor(CurChan, CurChan->Actor);
-				PendingNetActorChans.RemoveAt(i);
-			}
-		}
-	}
+		bool bBlockingProcess = IsBlockingProcessPresent();
 
-	// Prevent net connection timeout in developer mode
-	if (bDeveloperMode)
-	{
-		ResetConnTimeout();
-	}
-
-
-	PollProcessOutput();
-
-	// Update/save memory stats
-	UpdateProcessStats();
-}
-
-void UClientUnitTest::PollProcessOutput()
-{
-	for (auto CurHandle : ActiveProcesses)
-	{
-		if (CurHandle.IsValid())
+		if (!bBlockingProcess)
 		{
-			// Process any log/stdout data
-			FString LogDump = TEXT("");
-			bool bProcessedPipeRead = false;
+			ResetTimeout(TEXT("Blocking Process Reset"), true, 60);
 
-			// Need to iteratively grab pipe data, as it has a buffer (can miss data near process-end, for large logs, e.g. stack dumps)
-			while (true)
-			{
-				FString CurStdOut = FPlatformProcess::ReadPipe(CurHandle->ReadPipe);
-
-				if (CurStdOut.Len() > 0)
-				{
-					LogDump += CurStdOut;
-					bProcessedPipeRead = true;
-				}
-				// Sometimes large reads (typically > 4096 on the write side) clog the pipe buffer,
-				// and even looped reads won't receive it all, so when a large enough amount of data gets read, sleep momentarily
-				// (not ideal, but it works - spent a long time trying to find a better solution)
-				else if (bProcessedPipeRead && LogDump.Len() > 2048)
+			auto IsWaitingOnTimeout =
+				[&]()
 				{
-					bProcessedPipeRead = false;
+					return NextBlockingTimeout > FPlatformTime::Seconds();
+				};
 
-					FPlatformProcess::Sleep(0.01f);
-				}
-				else
-				{
-					break;
-				}
-			}
-
-			if (LogDump.Len() > 0)
+			if (bBlockingServerDelay && !IsWaitingOnTimeout())
 			{
-				// Every log line should start with an endline, so if one is missing, print that into the log as an error
-				bool bPartialRead = !LogDump.StartsWith(FString(LINE_TERMINATOR));
-				const TCHAR* PartialLog = TEXT("--MISSING ENDLINE - PARTIAL PIPE READ--");
-
-				// Now split up the log into multiple lines
-				TArray<FString> LogLines;
-				
-				// @todo #JohnB: Perhaps add support for both platforms line terminator, at some stage
-#if TARGET_UE4_CL < CL_STRINGPARSEARRAY
-				LogDump.ParseIntoArray(&LogLines, LINE_TERMINATOR, true);
-#else
-				LogDump.ParseIntoArray(LogLines, LINE_TERMINATOR, true);
-#endif
-
-
-				// For process logs, replace the timestamp/category with a tag (e.g. [SERVER]), and set a unique colour so it stands out
-				ELogTimes::Type bPreviousPrintLogTimes = GPrintLogTimes;
-				GPrintLogTimes = ELogTimes::None;
-
-				SET_WARN_COLOR(CurHandle->MainLogColor);
-
-				// Clear the engine-event log hook, to prevent duplication of the below log entry
-				UNIT_EVENT_CLEAR;
-
-				if (bPartialRead)
-				{
-					UE_LOG(None, Log, TEXT("%s"), PartialLog);
-				}
-
-				for (int LineIdx=0; LineIdx<LogLines.Num(); LineIdx++)
-				{
-					UE_LOG(None, Log, TEXT("%s%s"), *CurHandle->LogPrefix, *(LogLines[LineIdx]));
-				}
+				StartUnitTestServer();
 
-				// Restore the engine-event log hook
-				UNIT_EVENT_RESTORE;
-
-				CLEAR_WARN_COLOR();
-
-				GPrintLogTimes = bPreviousPrintLogTimes;
-
-
-				// Also output to the unit test log window
-				if (LogWindow.IsValid())
-				{
-					TSharedPtr<SLogWidget>& LogWidget = LogWindow->LogWidget;
-
-					if (LogWidget.IsValid())
-					{
-						if (bPartialRead)
-						{
-							LogWidget->AddLine(CurHandle->BaseLogType, MakeShareable(new FString(PartialLog)),
-												CurHandle->SlateLogColor);
-						}
-
-						for (int LineIdx=0; LineIdx<LogLines.Num(); LineIdx++)
-						{
-							FString LogLine = CurHandle->LogPrefix + LogLines[LineIdx];
-
-							LogWidget->AddLine(CurHandle->BaseLogType, MakeShareable(new FString(LogLine)), CurHandle->SlateLogColor);
-						}
-					}
-				}
-
-
-				// Now trigger notification of log lines (only if the unit test is not yet verified though)
-				if (VerificationState == EUnitTestVerification::Unverified)
-				{
-					NotifyProcessLog(CurHandle, LogLines);
-				}
-
-
-				// If the verification state has not changed, pass on the log lines to the error checking code
-				CheckOutputForError(CurHandle, LogLines);
+				bBlockingServerDelay = false;
+				NextBlockingTimeout = FPlatformTime::Seconds() + 10.0;
 			}
-		}
-	}
-}
-
-void UClientUnitTest::PostUnitTick(float DeltaTime)
-{
-	Super::PostUnitTick(DeltaTime);
-
-	TArray<TSharedPtr<FUnitTestProcess>> HandlesPendingShutdown;
-	bool bServerFinished = false;
-	bool bClientFinished  = false;
-
-	for (auto CurHandle : ActiveProcesses)
-	{
-		if (CurHandle.IsValid() && !FPlatformProcess::IsProcRunning(CurHandle->ProcessHandle))
-		{
-			FString LogMsg = FString::Printf(TEXT("Process '%s' has finished."), *CurHandle->ProcessTag);
-
-			UNIT_LOG(ELogType::StatusImportant, TEXT("%s"), *LogMsg);
-			UNIT_STATUS_LOG(ELogType::StatusVerbose, TEXT("%s"), *LogMsg);
-
-			HandlesPendingShutdown.Add(CurHandle);
-
 
-			if (ServerHandle.IsValid() && ServerHandle.HasSameObject(CurHandle.Get()))
+			if (bBlockingClientDelay && !IsWaitingOnTimeout())
 			{
-				bServerFinished = true;
+				ClientHandle = StartUnitTestClient(ServerAddress);
+
+				bBlockingClientDelay = false;
+				NextBlockingTimeout = FPlatformTime::Seconds() + 10.0;
 			}
-			else if (ClientHandle.IsValid() && ClientHandle.HasSameObject(CurHandle.Get()))
+
+			if (bBlockingFakeClientDelay && !IsWaitingOnTimeout())
 			{
-				bClientFinished = true;
+				ConnectFakeClient();
+
+				bBlockingFakeClientDelay = false;
+				NextBlockingTimeout = FPlatformTime::Seconds() + 10.0;
 			}
 		}
 	}
 
-	if (bServerFinished || bClientFinished)
+	if (!!(UnitTestFlags & EUnitTestFlags::NotifyNetActors) && PendingNetActorChans.Num() > 0 && UnitConn != NULL)
 	{
-		bool bProcessError = false;
-		FString UpdateMsg;
-
-		// If the server just finished, cleanup the fake client
-		if (bServerFinished)
+		for (int32 i=PendingNetActorChans.Num()-1; i>=0; i--)
 		{
-			FString LogMsg = FString::Printf(TEXT("Server process has finished, cleaning up fake client."));
-
-			UNIT_LOG(ELogType::StatusImportant, TEXT("%s"), *LogMsg);
-			UNIT_STATUS_LOG(ELogType::StatusVerbose, TEXT("%s"), *LogMsg);
-
-			if (UnitConn != NULL)
-			{
-				UnitConn->Close();
-			}
-
-			// Immediately cleanup the fake client (don't wait for end-of-life cleanup in CleanupUnitTest)
-			CleanupFakeClient();
-
+			UActorChannel* CurChan = Cast<UActorChannel>(UnitConn->Channels[PendingNetActorChans[i]]);
 
-			// If a server exit was unexpected, mark the unit test as broken
-			if (!(UnitTestFlags & EUnitTestFlags::IgnoreServerCrash) && VerificationState == EUnitTestVerification::Unverified)
+			if (CurChan != NULL && CurChan->Actor != NULL)
 			{
-				UpdateMsg = FString::Printf(TEXT("Unexpected server exit, marking unit test as needing update."));
-				bProcessError = true;
+				NotifyNetActor(CurChan, CurChan->Actor);
+				PendingNetActorChans.RemoveAt(i);
 			}
 		}
-
-		// If a client exit was unexpected, mark the unit test as broken
-		if (bClientFinished && !(UnitTestFlags & EUnitTestFlags::IgnoreClientCrash) &&
-			VerificationState == EUnitTestVerification::Unverified)
-		{
-			UpdateMsg = FString::Printf(TEXT("Unexpected client exit, marking unit test as needing update."));
-			bProcessError = true;
-		}
-
-
-		// If either the client/server finished, process the error
-		if (bProcessError)
-		{
-			UNIT_LOG(ELogType::StatusFailure | ELogType::StyleBold, TEXT("%s"), *UpdateMsg);
-			UNIT_STATUS_LOG(ELogType::StatusFailure | ELogType::StatusVerbose | ELogType::StyleBold, TEXT("%s"), *UpdateMsg);
-
-			VerificationState = EUnitTestVerification::VerifiedNeedsUpdate;
-		}
 	}
 
-
-	// Finally, clean up all handles pending shutdown (if they aren't already cleaned up)
-	if (HandlesPendingShutdown.Num() > 0)
+	// Prevent net connection timeout in developer mode
+	if (bDeveloperMode)
 	{
-		for (auto CurHandle : HandlesPendingShutdown)
-		{
-			if (CurHandle.IsValid())
-			{
-				ShutdownUnitTestProcess(CurHandle);
-			}
-		}
-
-		HandlesPendingShutdown.Empty();
+		ResetConnTimeout(120.f);
 	}
+
+	Super::UnitTick(DeltaTime);
 }
 
 bool UClientUnitTest::IsTickable() const
 {
 	bool bReturnVal = Super::IsTickable();
 
-	bReturnVal = bReturnVal && (ActiveProcesses.Num() > 0 ||
-					(!!(UnitTestFlags & EUnitTestFlags::NotifyNetActors) && PendingNetActorChans.Num() > 0));
+	bReturnVal = bReturnVal || bDeveloperMode || bBlockingServerDelay || bBlockingClientDelay || bBlockingFakeClientDelay ||
+					(!!(UnitTestFlags & EUnitTestFlags::NotifyNetActors) && PendingNetActorChans.Num() > 0);
 
 	return bReturnVal;
 }
 
-
-void UClientUnitTest::FinishDestroy()
-{
-	Super::FinishDestroy();
-
-	// Force close any processes still running
-	for (int32 i=ActiveProcesses.Num()-1; i>=0; i--)
-	{
-		if (ActiveProcesses[i].IsValid())
-		{
-			ShutdownUnitTestProcess(ActiveProcesses[i]);
-		}
-	}
-}
-
-void UClientUnitTest::ShutdownAfterError()
-{
-	Super::ShutdownAfterError();
-
-	// Force close any processes still running
-	for (int32 i=ActiveProcesses.Num()-1; i>=0; i--)
-	{
-		if (ActiveProcesses[i].IsValid())
-		{
-			ShutdownUnitTestProcess(ActiveProcesses[i]);
-		}
-	}
-}
-
diff --git a/Engine/Plugins/NetcodeUnitTest/NetcodeUnitTest/Source/NetcodeUnitTest/Private/NUTActor.cpp b/Engine/Plugins/NetcodeUnitTest/NetcodeUnitTest/Source/NetcodeUnitTest/Private/NUTActor.cpp
--- a/Engine/Plugins/NetcodeUnitTest/NetcodeUnitTest/Source/NetcodeUnitTest/Private/NUTActor.cpp
+++ b/Engine/Plugins/NetcodeUnitTest/NetcodeUnitTest/Source/NetcodeUnitTest/Private/NUTActor.cpp
@@ -86,6 +86,12 @@ bool ANUTActor::NotifyControlMessage(UNetConnection* Connection, uint8 MessageTy
 
 	if (MessageType == NMT_NUTControl)
 	{
+		// Some commands won't work without an owner, so if one is not set, set it now
+		if (Cast<APlayerController>(GetOwner()) == NULL)
+		{
+			UpdateOwner();
+		}
+
 		uint8 CmdType = 0;
 		FString Command;
 		FNetControlMessage<NMT_NUTControl>::Receive(Bunch, CmdType, Command);
@@ -332,48 +338,9 @@ void ANUTActor::Tick(float DeltaSeconds)
 		}
 
 		// Have the server set the owner, when appropriate
-		if ((Cast<APlayerController>(GetOwner()) == NULL || bClientTimedOut) && GEngine != NULL && CurNetMode != NM_Client)
+		if (Cast<APlayerController>(GetOwner()) == NULL || bClientTimedOut)
 		{
-			AGameState* GameState = CurWorld->GameState;
-
-			if (GameState != NULL)
-			{
-				// @todo #JohnBReview: You want this to only happen if no remote players are present (perhaps give all players control instead,
-				//				if this becomes a problem - requires setting things up differently though)
-				if (CurNetMode == NM_ListenServer || CurNetMode == NM_Standalone)
-				{
-					for (FLocalPlayerIterator It(GEngine, CurWorld); It; ++It)
-					{
-						if (It->PlayerController != NULL)
-						{
-							// Reset LastAliveTime, to give client a chance to send initial 'alive' RPC
-							LastAliveTime = CurWorld->RealTimeSeconds;
-
-							SetOwner(It->PlayerController);
-							break;
-						}
-					}
-				}
-
-				for (int i=0; i<GameState->PlayerArray.Num(); i++)
-				{
-					APlayerController* PC = Cast<APlayerController>(GameState->PlayerArray[i] != NULL ?
-																	GameState->PlayerArray[i]->GetOwner() : NULL);
-
-					if (PC != NULL && PC != GetOwner() && Cast<UNetConnection>(PC->Player) != NULL)
-					{
-						UE_LOG(LogUnitTest, Log, TEXT("Setting NUTActor owner to: %s (%s)"), *PC->GetName(),
-							*GameState->PlayerArray[i]->PlayerName);
-
-						// Reset LastAliveTime, to give client a chance to send initial 'alive' RPC
-						LastAliveTime = CurWorld->RealTimeSeconds;
-
-						SetOwner(PC);
-
-						break;
-					}
-				}
-			}
+			UpdateOwner();
 		}
 
 
@@ -429,6 +396,52 @@ void ANUTActor::Tick(float DeltaSeconds)
 	}
 }
 
+void ANUTActor::UpdateOwner()
+{
+	UWorld* CurWorld = GetWorld();
+	ENetMode CurNetMode = GEngine != NULL ? GEngine->GetNetMode(CurWorld) : NM_Standalone;
+	AGameState* GameState = CurWorld->GameState;
+
+	if (GameState != NULL && GEngine != NULL && CurNetMode != NM_Client)
+	{
+		// @todo #JohnBReview: You want this to only happen if no remote players are present (perhaps give all players control instead,
+		//				if this becomes a problem - requires setting things up differently though)
+		if (CurNetMode == NM_ListenServer || CurNetMode == NM_Standalone)
+		{
+			for (FLocalPlayerIterator It(GEngine, CurWorld); It; ++It)
+			{
+				if (It->PlayerController != NULL)
+				{
+					// Reset LastAliveTime, to give client a chance to send initial 'alive' RPC
+					LastAliveTime = CurWorld->RealTimeSeconds;
+
+					SetOwner(It->PlayerController);
+					break;
+				}
+			}
+		}
+
+		for (int i=0; i<GameState->PlayerArray.Num(); i++)
+		{
+			APlayerController* PC = Cast<APlayerController>(GameState->PlayerArray[i] != NULL ?
+															GameState->PlayerArray[i]->GetOwner() : NULL);
+
+			if (PC != NULL && PC != GetOwner() && Cast<UNetConnection>(PC->Player) != NULL)
+			{
+				UE_LOG(LogUnitTest, Log, TEXT("Setting NUTActor owner to: %s (%s)"), *PC->GetName(),
+					*GameState->PlayerArray[i]->PlayerName);
+
+				// Reset LastAliveTime, to give client a chance to send initial 'alive' RPC
+				LastAliveTime = CurWorld->RealTimeSeconds;
+
+				SetOwner(PC);
+
+				break;
+			}
+		}
+	}
+}
+
 void ANUTActor::HookNetDriver(UNetDriver* TargetNetDriver)
 {
 	if (TargetNetDriver != NULL)
diff --git a/Engine/Plugins/NetcodeUnitTest/NetcodeUnitTest/Source/NetcodeUnitTest/Private/NUTUtil.cpp b/Engine/Plugins/NetcodeUnitTest/NetcodeUnitTest/Source/NetcodeUnitTest/Private/NUTUtil.cpp
--- a/Engine/Plugins/NetcodeUnitTest/NetcodeUnitTest/Source/NetcodeUnitTest/Private/NUTUtil.cpp
+++ b/Engine/Plugins/NetcodeUnitTest/NetcodeUnitTest/Source/NetcodeUnitTest/Private/NUTUtil.cpp
@@ -16,7 +16,9 @@ void NUTUtil::GetUnitTestClassDefList(TArray<UUnitTest*>& OutUnitTestClassDefaul
 {
 	for (TObjectIterator<UClass> It; It; ++It)
 	{
-		if (It->IsChildOf(UUnitTest::StaticClass()) && *It != UUnitTest::StaticClass() && *It != UClientUnitTest::StaticClass())
+		// @todo #JohnBRefactor: Move all 'abstract class' checks to a central location, as you do this in multiple locations
+		if (It->IsChildOf(UUnitTest::StaticClass()) && *It != UUnitTest::StaticClass() && *It != UClientUnitTest::StaticClass() &&
+			*It != UProcessUnitTest::StaticClass())
 		{
 			UUnitTest* CurDefault = Cast<UUnitTest>(It->GetDefaultObject());
 
diff --git a/Engine/Plugins/NetcodeUnitTest/NetcodeUnitTest/Source/NetcodeUnitTest/Private/Net/NUTUtilNet.cpp b/Engine/Plugins/NetcodeUnitTest/NetcodeUnitTest/Source/NetcodeUnitTest/Private/Net/NUTUtilNet.cpp
--- a/Engine/Plugins/NetcodeUnitTest/NetcodeUnitTest/Source/NetcodeUnitTest/Private/Net/NUTUtilNet.cpp
+++ b/Engine/Plugins/NetcodeUnitTest/NetcodeUnitTest/Source/NetcodeUnitTest/Private/Net/NUTUtilNet.cpp
@@ -362,7 +362,11 @@ bool NUTNet::CreateFakePlayer(UWorld* InWorld, UNetDriver*& InNetDriver, FString
 #if TARGET_UE4_CL >= CL_FNETWORKVERSION
 			// Starting with 4.8.0, the network protocol has changed slightly
 			// @todo JohnB: Refactor this, to toggle at compile time only, based on CL (might require more accurate UT integrate CLs)
+#if TARGET_UE4_CL >= CL_FENGINEVERSION
 			FString VersionStr = FEngineVersion::Current().ToString(EVersionComponent::Minor);
+#else
+			FString VersionStr = GEngineVersion.ToString(EVersionComponent::Minor);
+#endif
 			int32 VersionDelim = VersionStr.Find(TEXT("."));
 			int32 MajorVersion = FCString::Atoi(*VersionStr.Left(VersionDelim));
 			int32 MinorVersion = FCString::Atoi(*VersionStr.Mid(VersionDelim+1));
diff --git a/Engine/Plugins/NetcodeUnitTest/NetcodeUnitTest/Source/NetcodeUnitTest/Private/Net/UnitTestPackageMap.cpp b/Engine/Plugins/NetcodeUnitTest/NetcodeUnitTest/Source/NetcodeUnitTest/Private/Net/UnitTestPackageMap.cpp
--- a/Engine/Plugins/NetcodeUnitTest/NetcodeUnitTest/Source/NetcodeUnitTest/Private/Net/UnitTestPackageMap.cpp
+++ b/Engine/Plugins/NetcodeUnitTest/NetcodeUnitTest/Source/NetcodeUnitTest/Private/Net/UnitTestPackageMap.cpp
@@ -16,24 +16,24 @@ UUnitTestPackageMap::UUnitTestPackageMap(const FObjectInitializer& ObjectInitial
 {
 }
 
-bool UUnitTestPackageMap::SerializeObject(FArchive& Ar, UClass* Class, UObject*& Obj, FNetworkGUID* OutNetGUID/*=NULL */)
+bool UUnitTestPackageMap::SerializeObject(FArchive& Ar, UClass* InClass, UObject*& Obj, FNetworkGUID* OutNetGUID/*=NULL */)
 {
 	bool bReturnVal = false;
 
-	bReturnVal = Super::SerializeObject(Ar, Class, Obj, OutNetGUID);
+	bReturnVal = Super::SerializeObject(Ar, InClass, Obj, OutNetGUID);
 
 	if (bWithinSerializeNewActor)
 	{
 		// This indicates that SerializeObject has failed to find an existing instance when trying to serialize an actor,
 		// so it will be spawned clientside later on (after the archetype is serialized) instead.
 		// These spawns count as undesired clientside code execution, so filter them through NotifyAllowNetActor.
-		if (Class == AActor::StaticClass() && Obj == NULL)
+		if (InClass == AActor::StaticClass() && Obj == NULL)
 		{
 			bPendingArchetypeSpawn = true;
 		}
 		// This indicates that a new actor archetype has just been serialized (which may or may not be during actor channel init);
 		// this is the first place we know the type of a replicated actor (in an actor channel or otherwise), but BEFORE it is spawned
-		else if ((GIsInitializingActorChan || bPendingArchetypeSpawn) && Class == UObject::StaticClass() && Obj != NULL)
+		else if ((GIsInitializingActorChan || bPendingArchetypeSpawn) && InClass == UObject::StaticClass() && Obj != NULL)
 		{
 			UUnitTestNetConnection* UnitConn = Cast<UUnitTestNetConnection>(GActiveReceiveUnitConnection);
 			bool bAllowActor = false;
diff --git a/Engine/Plugins/NetcodeUnitTest/NetcodeUnitTest/Source/NetcodeUnitTest/Private/ProcessUnitTest.cpp b/Engine/Plugins/NetcodeUnitTest/NetcodeUnitTest/Source/NetcodeUnitTest/Private/ProcessUnitTest.cpp
new file mode 100644
--- /dev/null
+++ b/Engine/Plugins/NetcodeUnitTest/NetcodeUnitTest/Source/NetcodeUnitTest/Private/ProcessUnitTest.cpp
@@ -0,0 +1,613 @@
+// Copyright 1998-2015 Epic Games, Inc. All Rights Reserved.
+
+#include "NetcodeUnitTestPCH.h"
+
+#include "ProcessUnitTest.h"
+
+#include "UnitTestEnvironment.h"
+
+
+/**
+ * UProcessUnitTest
+ */
+UProcessUnitTest::UProcessUnitTest(const FObjectInitializer& ObjectInitializer)
+	: Super(ObjectInitializer)
+	, ActiveProcesses()
+	, OnSuspendStateChange()
+	, LastBlockingProcessCheck(0)
+{
+}
+
+void UProcessUnitTest::NotifyProcessSuspendState(TWeakPtr<FUnitTestProcess> InProcess, ESuspendState InSuspendState)
+{
+	if (InProcess.IsValid())
+	{
+		InProcess.Pin()->SuspendState = InSuspendState;
+	}
+}
+
+
+void UProcessUnitTest::CleanupUnitTest()
+{
+	for (int32 i=ActiveProcesses.Num()-1; i>=0; i--)
+	{
+		if (ActiveProcesses[i].IsValid())
+		{
+			ShutdownUnitTestProcess(ActiveProcesses[i]);
+		}
+	}
+
+	Super::CleanupUnitTest();
+}
+
+TWeakPtr<FUnitTestProcess> UProcessUnitTest::StartUnitTestProcess(FString Path, FString InCommandline, bool bMinimized/*=true*/)
+{
+	TSharedPtr<FUnitTestProcess> ReturnVal = MakeShareable(new FUnitTestProcess());
+
+	verify(FPlatformProcess::CreatePipe(ReturnVal->ReadPipe, ReturnVal->WritePipe));
+
+	UNIT_LOG(ELogType::StatusImportant, TEXT("Starting process with parameters: %s"), *InCommandline);
+
+	ReturnVal->ProcessHandle = FPlatformProcess::CreateProc(*Path, *InCommandline, true, bMinimized, false, &ReturnVal->ProcessID,
+															0, NULL, ReturnVal->WritePipe);
+
+	if (ReturnVal->ProcessHandle.IsValid())
+	{
+		ReturnVal->ProcessTag = FString::Printf(TEXT("Process_%i"), ReturnVal->ProcessID);
+		ReturnVal->LogPrefix = FString::Printf(TEXT("[PROC%i]"), ReturnVal->ProcessID);
+
+		ActiveProcesses.Add(ReturnVal);
+	}
+	else
+	{
+		UNIT_LOG(ELogType::StatusFailure, TEXT("Failed to start process"));
+	}
+
+	return ReturnVal;
+}
+
+TWeakPtr<FUnitTestProcess> UProcessUnitTest::StartUE4UnitTestProcess(FString InCommandline, bool bMinimized/*=true*/)
+{
+	TWeakPtr<FUnitTestProcess> ReturnVal = NULL;
+	FString GamePath = FPlatformProcess::GenerateApplicationPath(FApp::GetName(), FApp::GetBuildConfiguration());
+
+	ReturnVal = StartUnitTestProcess(GamePath, InCommandline, bMinimized);
+
+	if (ReturnVal.IsValid())
+	{
+		auto CurHandle = ReturnVal.Pin();
+
+		if (CurHandle->ProcessHandle.IsValid())
+		{
+			CurHandle->ProcessTag = FString::Printf(TEXT("UE4_%i"), CurHandle->ProcessID);
+			CurHandle->LogPrefix = FString::Printf(TEXT("[UE4%i]"), CurHandle->ProcessID);
+		}
+	}
+
+	return ReturnVal;
+}
+
+void UProcessUnitTest::ShutdownUnitTestProcess(TSharedPtr<FUnitTestProcess> InHandle)
+{
+	if (InHandle->ProcessHandle.IsValid())
+	{
+		FString LogMsg = FString::Printf(TEXT("Shutting down process '%s'."), *InHandle->ProcessTag);
+
+		UNIT_LOG(ELogType::StatusImportant, TEXT("%s"), *LogMsg);
+		UNIT_STATUS_LOG(ELogType::StatusVerbose, TEXT("%s"), *LogMsg);
+
+
+		// @todo #JohnBHighPri: Restore 'true' here, once the issue where killing child processes sometimes kills all processes,
+		//						is fixed.
+		FPlatformProcess::TerminateProc(InHandle->ProcessHandle);//, true);
+
+#if TARGET_UE4_CL < CL_CLOSEPROC
+		InHandle->ProcessHandle.Close();
+#else
+		FPlatformProcess::CloseProc(InHandle->ProcessHandle);
+#endif
+	}
+
+	FPlatformProcess::ClosePipe(InHandle->ReadPipe, InHandle->WritePipe);
+	InHandle->ReadPipe = NULL;
+	InHandle->WritePipe = NULL;
+
+	ActiveProcesses.Remove(InHandle);
+
+	// Print out any detected error logs
+	if (InHandle->ErrorLogStage != EErrorLogStage::ELS_NoError && InHandle->ErrorText.Num() > 0)
+	{
+		PrintUnitTestProcessErrors(InHandle);
+	}
+}
+
+void UProcessUnitTest::PrintUnitTestProcessErrors(TSharedPtr<FUnitTestProcess> InHandle)
+{
+	FString LogMsg = FString::Printf(TEXT("Detected a crash in process '%s':"), *InHandle->ProcessTag);
+
+	UNIT_LOG(ELogType::StatusImportant, TEXT("%s"), *LogMsg);
+	UNIT_STATUS_LOG(ELogType::StatusImportant, TEXT("%s (click 'Advanced Summary' for more detail)"), *LogMsg);
+
+
+	uint32 CallstackCount = 0;
+
+	for (auto CurErrorLine : InHandle->ErrorText)
+	{
+		ELogType CurLogType = ELogType::None;
+
+		if (CurErrorLine.Stage == EErrorLogStage::ELS_ErrorStart)
+		{
+			CurLogType = ELogType::StatusAdvanced;
+		}
+		else if (CurErrorLine.Stage == EErrorLogStage::ELS_ErrorDesc)
+		{
+			CurLogType = ELogType::StatusImportant | ELogType::StatusAdvanced | ELogType::StyleBold;
+		}
+		else if (CurErrorLine.Stage == EErrorLogStage::ELS_ErrorCallstack)
+		{
+			CurLogType = ELogType::StatusAdvanced;
+
+			// Include the first five callstack lines in the main summary printout
+			if (CallstackCount < 5)
+			{
+				CurLogType |= ELogType::StatusImportant;
+			}
+
+			CallstackCount++;
+		}
+		else if (CurErrorLine.Stage == EErrorLogStage::ELS_ErrorExit)
+		{
+			continue;
+		}
+		else
+		{
+			// Make sure there's a check for all error log stages
+			UNIT_ASSERT(false);
+		}
+
+		// For the current unit-test log window, all entries are marked 'StatusImportant'
+		UNIT_LOG(CurLogType | ELogType::StatusImportant, TEXT("%s"), *CurErrorLine.Line);
+
+		UNIT_STATUS_LOG(CurLogType, TEXT("%s"), *CurErrorLine.Line);
+	}
+}
+
+void UProcessUnitTest::UpdateProcessStats()
+{
+	SIZE_T TotalProcessMemoryUsage = 0;
+
+	for (auto CurHandle : ActiveProcesses)
+	{
+		// Check unit test memory stats, and update if necessary (NOTE: Processes not guaranteed to still be running)
+		if (CurHandle.IsValid() && CurHandle->ProcessID != 0)
+		{
+			SIZE_T ProcessMemUsage = 0;
+
+			if (FPlatformProcess::GetApplicationMemoryUsage(CurHandle->ProcessID, &ProcessMemUsage) && ProcessMemUsage != 0)
+			{
+				TotalProcessMemoryUsage += ProcessMemUsage;
+			}
+		}
+	}
+
+	if (TotalProcessMemoryUsage > 0)
+	{
+		CurrentMemoryUsage = TotalProcessMemoryUsage;
+
+		float RunningTime = (float)(FPlatformTime::Seconds() - StartTime);
+
+		// Update saved memory usage stats, to help with tracking estimated memory usage, on future unit test launches
+		if (CurrentMemoryUsage > PeakMemoryUsage)
+		{
+			if (PeakMemoryUsage == 0)
+			{
+				bFirstTimeStats = true;
+			}
+
+			PeakMemoryUsage = CurrentMemoryUsage;
+
+			// Reset TimeToPeakMem
+			TimeToPeakMem = RunningTime;
+
+			SaveConfig();
+		}
+		// Check if we have hit a new record, for the time it takes to get within 90% of PeakMemoryUsage
+		else if (RunningTime < TimeToPeakMem && (CurrentMemoryUsage * 100) >= (PeakMemoryUsage * 90))
+		{
+			TimeToPeakMem = RunningTime;
+			SaveConfig();
+		}
+	}
+}
+
+void UProcessUnitTest::CheckOutputForError(TSharedPtr<FUnitTestProcess> InProcess, const TArray<FString>& InLines)
+{
+	// Perform error parsing, for UE4 processes
+	if (InProcess->ProcessTag.StartsWith(TEXT("UE4_")))
+	{
+		// Array of log messages that can indicate the start of an error
+		static const TArray<FString> ErrorStartLogs = TArrayBuilder<FString>()
+			.Add(FString(TEXT("Windows GetLastError: ")))
+			.Add(FString(TEXT("=== Critical error: ===")));
+
+
+		for (FString CurLine : InLines)
+		{
+			// Using 'ContainsByPredicate' as an iterator
+			const auto CheckForErrorLog =
+				[&CurLine](const FString& ErrorLine)
+				{
+					return CurLine.Contains(ErrorLine);
+				};
+
+			// Search for the beginning of an error log message
+			if (InProcess->ErrorLogStage == EErrorLogStage::ELS_NoError)
+			{
+				if (ErrorStartLogs.ContainsByPredicate(CheckForErrorLog))
+				{
+					InProcess->ErrorLogStage = EErrorLogStage::ELS_ErrorStart;
+
+					// Reset the timeout for both the unit test and unit test connection here, as callstack logs are prone to failure
+					//	(do it for at least two minutes as well, as it can take a long time to get a crashlog)
+					ResetTimeout(TEXT("Detected crash."), true, 120);
+				}
+			}
+
+
+			if (InProcess->ErrorLogStage != EErrorLogStage::ELS_NoError)
+			{
+				// Regex pattern for matching callstack logs - matches:
+				//	" (0x000007fefe22cacd) + 0 bytes"
+				//	" {0x000007fefe22cacd} + 0 bytes"
+				const FRegexPattern CallstackPattern(TEXT("\\s[\\(|\\{]0x[0-9,a-f]+[\\)|\\}] \\+ [0-9]+ bytes"));
+
+				// Matches:
+				//	"ntldll.dll"
+				const FRegexPattern AltCallstackPattern(TEXT("^[a-z,A-Z,0-9,\\-,_]+\\.[exe|dll]"));
+
+
+				// Check for the beginning of description logs
+				if (InProcess->ErrorLogStage == EErrorLogStage::ELS_ErrorStart &&
+					!ErrorStartLogs.ContainsByPredicate(CheckForErrorLog))
+				{
+					InProcess->ErrorLogStage = EErrorLogStage::ELS_ErrorDesc;
+				}
+
+				// Check-for/verify callstack logs
+				if (InProcess->ErrorLogStage == EErrorLogStage::ELS_ErrorDesc ||
+					InProcess->ErrorLogStage == EErrorLogStage::ELS_ErrorCallstack)
+				{
+					FRegexMatcher CallstackMatcher(CallstackPattern, CurLine);
+					FRegexMatcher AltCallstackMatcher(AltCallstackPattern, CurLine);
+
+					if (CallstackMatcher.FindNext() || AltCallstackMatcher.FindNext())
+					{
+						InProcess->ErrorLogStage = EErrorLogStage::ELS_ErrorCallstack;
+					}
+					else if (InProcess->ErrorLogStage == EErrorLogStage::ELS_ErrorCallstack)
+					{
+						InProcess->ErrorLogStage = EErrorLogStage::ELS_ErrorExit;
+					}
+				}
+
+				// The rest of the lines, after callstack parsing, should be ELS_ErrorExit logs - no need to verify these
+
+
+				InProcess->ErrorText.Add(FErrorLog(InProcess->ErrorLogStage, CurLine));
+			}
+		}
+	}
+}
+
+void UProcessUnitTest::UnitTick(float DeltaTime)
+{
+	Super::UnitTick(DeltaTime);
+
+	// Detect processes which indicate a blocking task, and reset timeouts while this is occurring
+	if ((FPlatformTime::Seconds() - LastBlockingProcessCheck) > 5.0 && IsBlockingProcessPresent(true))
+	{
+		ResetTimeout(TEXT("Blocking Child Process"), true, 60);
+
+		LastBlockingProcessCheck = FPlatformTime::Seconds();
+	}
+
+
+	PollProcessOutput();
+
+	// Update/save memory stats
+	UpdateProcessStats();
+}
+
+void UProcessUnitTest::PostUnitTick(float DeltaTime)
+{
+	Super::PostUnitTick(DeltaTime);
+
+	TArray<TSharedPtr<FUnitTestProcess>> HandlesPendingShutdown;
+
+	for (auto CurHandle : ActiveProcesses)
+	{
+		if (CurHandle.IsValid() && !FPlatformProcess::IsProcRunning(CurHandle->ProcessHandle))
+		{
+			FString LogMsg = FString::Printf(TEXT("Process '%s' has finished."), *CurHandle->ProcessTag);
+
+			UNIT_LOG(ELogType::StatusImportant, TEXT("%s"), *LogMsg);
+			UNIT_STATUS_LOG(ELogType::StatusVerbose, TEXT("%s"), *LogMsg);
+
+			HandlesPendingShutdown.Add(CurHandle);
+		}
+	}
+
+	for (auto CurProcHandle : HandlesPendingShutdown)
+	{
+		NotifyProcessFinished(CurProcHandle);
+	}
+
+
+	// Finally, clean up all handles pending shutdown (if they aren't already cleaned up)
+	if (HandlesPendingShutdown.Num() > 0)
+	{
+		for (auto CurHandle : HandlesPendingShutdown)
+		{
+			if (CurHandle.IsValid())
+			{
+				ShutdownUnitTestProcess(CurHandle);
+			}
+		}
+
+		HandlesPendingShutdown.Empty();
+	}
+}
+
+bool UProcessUnitTest::IsTickable() const
+{
+	bool bReturnVal = Super::IsTickable();
+
+	bReturnVal = bReturnVal || ActiveProcesses.Num() > 0;
+
+	return bReturnVal;
+}
+
+
+void UProcessUnitTest::PollProcessOutput()
+{
+	for (auto CurHandle : ActiveProcesses)
+	{
+		if (CurHandle.IsValid())
+		{
+			// Process any log/stdout data
+			FString LogDump = TEXT("");
+			bool bProcessedPipeRead = false;
+
+			// Need to iteratively grab pipe data, as it has a buffer (can miss data near process-end, for large logs, e.g. stack dumps)
+			while (true)
+			{
+				FString CurStdOut = FPlatformProcess::ReadPipe(CurHandle->ReadPipe);
+
+				if (CurStdOut.Len() > 0)
+				{
+					LogDump += CurStdOut;
+					bProcessedPipeRead = true;
+				}
+				// Sometimes large reads (typically > 4096 on the write side) clog the pipe buffer,
+				// and even looped reads won't receive it all, so when a large enough amount of data gets read, sleep momentarily
+				// (not ideal, but it works - spent a long time trying to find a better solution)
+				else if (bProcessedPipeRead && LogDump.Len() > 2048)
+				{
+					bProcessedPipeRead = false;
+
+					FPlatformProcess::Sleep(0.01f);
+				}
+				else
+				{
+					break;
+				}
+			}
+
+			if (LogDump.Len() > 0)
+			{
+				// Every log line should start with an endline, so if one is missing, print that into the log as an error
+				bool bPartialRead = !LogDump.StartsWith(FString(LINE_TERMINATOR));
+				const TCHAR* PartialLog = TEXT("--MISSING ENDLINE - PARTIAL PIPE READ--");
+
+				// Now split up the log into multiple lines
+				TArray<FString> LogLines;
+				
+				// @todo #JohnB: Perhaps add support for both platforms line terminator, at some stage
+#if TARGET_UE4_CL < CL_STRINGPARSEARRAY
+				LogDump.ParseIntoArray(&LogLines, LINE_TERMINATOR, true);
+#else
+				LogDump.ParseIntoArray(LogLines, LINE_TERMINATOR, true);
+#endif
+
+
+				// For process logs, replace the timestamp/category with a tag (e.g. [SERVER]), and set a unique colour so it stands out
+				ELogTimes::Type bPreviousPrintLogTimes = GPrintLogTimes;
+				GPrintLogTimes = ELogTimes::None;
+
+				SET_WARN_COLOR(CurHandle->MainLogColor);
+
+				// Clear the engine-event log hook, to prevent duplication of the below log entry
+				UNIT_EVENT_CLEAR;
+
+				if (bPartialRead)
+				{
+					UE_LOG(None, Log, TEXT("%s"), PartialLog);
+				}
+
+				for (int LineIdx=0; LineIdx<LogLines.Num(); LineIdx++)
+				{
+					UE_LOG(None, Log, TEXT("%s%s"), *CurHandle->LogPrefix, *(LogLines[LineIdx]));
+				}
+
+				// Restore the engine-event log hook
+				UNIT_EVENT_RESTORE;
+
+				CLEAR_WARN_COLOR();
+
+				GPrintLogTimes = bPreviousPrintLogTimes;
+
+
+				// Also output to the unit test log window
+				if (LogWindow.IsValid())
+				{
+					TSharedPtr<SLogWidget>& LogWidget = LogWindow->LogWidget;
+
+					if (LogWidget.IsValid())
+					{
+						if (bPartialRead)
+						{
+							LogWidget->AddLine(CurHandle->BaseLogType, MakeShareable(new FString(PartialLog)),
+												CurHandle->SlateLogColor);
+						}
+
+						for (int LineIdx=0; LineIdx<LogLines.Num(); LineIdx++)
+						{
+							FString LogLine = CurHandle->LogPrefix + LogLines[LineIdx];
+
+							LogWidget->AddLine(CurHandle->BaseLogType, MakeShareable(new FString(LogLine)), CurHandle->SlateLogColor);
+						}
+					}
+				}
+
+
+				// Now trigger notification of log lines (only if the unit test is not yet verified though)
+				if (VerificationState == EUnitTestVerification::Unverified)
+				{
+					NotifyProcessLog(CurHandle, LogLines);
+				}
+
+
+				// If the verification state has not changed, pass on the log lines to the error checking code
+				CheckOutputForError(CurHandle, LogLines);
+			}
+		}
+	}
+}
+
+bool UProcessUnitTest::IsBlockingProcessPresent(bool bLogIfFound/*=false*/)
+{
+	bool bReturnVal = false;
+	const TArray<FString>* BlockingProcesses = NULL;
+
+	UnitEnv->GetProgressBlockingProcesses(BlockingProcesses);
+
+
+	// Map all process id's and parent ids, to generate the entire child process tree
+	TMap<uint32, uint32> ProcMap;
+
+	{
+		FPlatformProcess::FProcEnumerator ProcIt;
+
+		while (ProcIt.MoveNext())
+		{
+			auto CurProc = ProcIt.GetCurrent();
+			ProcMap.Add(CurProc.GetPID(), CurProc.GetParentPID());
+		}
+	}
+
+
+	TArray<uint32> ChildProcs;
+	ChildProcs.Add(FPlatformProcess::GetCurrentProcessId());
+
+	for (int32 i=0; i<ChildProcs.Num(); i++)
+	{
+		uint32 SearchParentId = ChildProcs[i];
+
+		for (TMap<uint32, uint32>::TConstIterator It(ProcMap); It; ++It)
+		{
+			if (It.Value() == SearchParentId)
+			{
+				ChildProcs.Add(It.Key());
+			}
+		}
+	}
+
+
+	// Now check if any of those child processes are a blocking process
+	TArray<FString> BlockingProcNames;
+
+	{
+		FPlatformProcess::FProcEnumerator ProcIt;
+
+		while (ProcIt.MoveNext())
+		{
+			auto CurProc = ProcIt.GetCurrent();
+
+			if (ChildProcs.Contains(CurProc.GetPID()))
+			{
+				FString CurProcName = CurProc.GetName();
+				int32 Delim = CurProcName.Find(TEXT("."), ESearchCase::CaseSensitive, ESearchDir::FromEnd);
+
+				if (Delim != INDEX_NONE)
+				{
+					CurProcName = CurProcName.Left(Delim);
+				}
+
+				if (BlockingProcesses->Contains(CurProcName))
+				{
+					bReturnVal = true;
+					BlockingProcNames.AddUnique(CurProcName);
+				}
+			}
+		}
+	}
+
+	if (bLogIfFound && bReturnVal)
+	{
+		FString ProcString;
+
+		for (auto CurProc : BlockingProcNames)
+		{
+			if (ProcString.Len() != 0)
+			{
+				ProcString += TEXT(", ");
+			}
+
+			ProcString += CurProc;
+		}
+
+		UNIT_LOG(ELogType::StatusImportant, TEXT("Detected blocking child process '%s'."), *ProcString);
+
+
+		// Longstanding source of issues, with running unit tests - haven't found a good fully-automated solution yet
+		if (BlockingProcNames.Contains(TEXT("ShaderCompileWorker")))
+		{
+			UNIT_LOG(ELogType::StatusImportant,
+					TEXT("Recommend generating shaders separately before running unit tests, to avoid bugs and unit test failure."));
+
+			UNIT_LOG(ELogType::StatusImportant,
+					TEXT("One way to do this (takes a long time), is through: UE4Editor-Cmd YourGame -run=DerivedDataCache -fill"));
+		}
+	}
+
+	return bReturnVal;
+}
+
+void UProcessUnitTest::FinishDestroy()
+{
+	Super::FinishDestroy();
+
+	// Force close any processes still running
+	for (int32 i=ActiveProcesses.Num()-1; i>=0; i--)
+	{
+		if (ActiveProcesses[i].IsValid())
+		{
+			ShutdownUnitTestProcess(ActiveProcesses[i]);
+		}
+	}
+}
+
+void UProcessUnitTest::ShutdownAfterError()
+{
+	Super::ShutdownAfterError();
+
+	// Force close any processes still running
+	for (int32 i=ActiveProcesses.Num()-1; i>=0; i--)
+	{
+		if (ActiveProcesses[i].IsValid())
+		{
+			ShutdownUnitTestProcess(ActiveProcesses[i]);
+		}
+	}
+}
+
diff --git a/Engine/Plugins/NetcodeUnitTest/NetcodeUnitTest/Source/NetcodeUnitTest/Private/UI/SLogWidget.cpp b/Engine/Plugins/NetcodeUnitTest/NetcodeUnitTest/Source/NetcodeUnitTest/Private/UI/SLogWidget.cpp
--- a/Engine/Plugins/NetcodeUnitTest/NetcodeUnitTest/Source/NetcodeUnitTest/Private/UI/SLogWidget.cpp
+++ b/Engine/Plugins/NetcodeUnitTest/NetcodeUnitTest/Source/NetcodeUnitTest/Private/UI/SLogWidget.cpp
@@ -29,8 +29,8 @@ IMPLEMENT_GET_PRIVATE_VAR(SEditableTextBox, EditableText, TSharedPtr<SEditableTe
 // Enable access to SButton.Style
 IMPLEMENT_GET_PRIVATE_VAR(SButton, Style, const FButtonStyle*);
 
-// Enable access to SDockTab::GetCurrentStyle, using the CALL_PRIVATE macro
-// IMPLEMENT_GET_PRIVATE_FUNC_CONST(SDockTab, GetCurrentStyle, const FDockTabStyle&, void, const);
+// Enable access to SDockTab::GetCurrentStyle, using the CALL_PROTECTED macro
+IMPLEMENT_GET_PROTECTED_FUNC_CONST(SDockTab, GetCurrentStyle, const FDockTabStyle&, void,, const);
 
 
 // @todo #JohnBRefactorUI: Perhaps move widget searching to a NUTSlate.h file, or such?
@@ -712,21 +712,21 @@ TSharedRef<SDockTab> SLogWidget::SpawnLogTab(const FSpawnTabArgs& InSpawnTabArgs
 							{
 								// Various types of special font formatting
 								FString FontPath;
-								ELogType LogType = Item->LogType;
+								ELogType CurLogType = Item->LogType;
 
-								if (!!(LogType & ELogType::StyleMonospace))
+								if (!!(CurLogType & ELogType::StyleMonospace))
 								{
 									FontPath = FPaths::EngineContentDir() / TEXT("Slate/Fonts/DroidSansMono.ttf");
 								}
-								else if (!!(LogType & ELogType::StyleBold) && !!(LogType & ELogType::StyleItalic))
+								else if (!!(CurLogType & ELogType::StyleBold) && !!(CurLogType & ELogType::StyleItalic))
 								{
 									FontPath = FPaths::EngineContentDir() / TEXT("Editor/Slate/Fonts/Roboto-BoldCondensedItalic.ttf");
 								}
-								else if (!!(LogType & ELogType::StyleBold))
+								else if (!!(CurLogType & ELogType::StyleBold))
 								{
 									FontPath = FPaths::EngineContentDir() / TEXT("Slate/Fonts/Roboto-Bold.ttf");
 								}
-								else if (!!(LogType & ELogType::StyleItalic))
+								else if (!!(CurLogType & ELogType::StyleItalic))
 								{
 									FontPath = FPaths::EngineContentDir() / TEXT("Editor/Slate/Fonts/Roboto-Italic.ttf");
 								}
@@ -741,7 +741,7 @@ TSharedRef<SDockTab> SLogWidget::SpawnLogTab(const FSpawnTabArgs& InSpawnTabArgs
 
 
 								// Pseudo-underline; just adds newline, and then underlines with lots of ----
-								if (!!(LogType & ELogType::StyleUnderline) && RenderText.Len() > 0)
+								if (!!(CurLogType & ELogType::StyleUnderline) && RenderText.Len() > 0)
 								{
 									const TCHAR* TextToUnderline = *RenderText;
 									const TCHAR* UnderlineEnd = TextToUnderline + RenderText.Len()-1;
@@ -968,9 +968,9 @@ TSharedRef<SDockTab> SLogWidget::SpawnLogTab(const FSpawnTabArgs& InSpawnTabArgs
 				{
 					TSharedRef<SButton> CurButton = StaticCastSharedRef<SButton>(InWidget);
 					const FButtonStyle* ButtonStyle = GET_PRIVATE(SButton, CurButton, Style);
-// 					const FDockTabStyle& TabStyle = CALL_PRIVATE(SDockTab, ReturnVal, GetCurrentStyle)();
-// 
-// 					bFound = ButtonStyle == &TabStyle.CloseButtonStyle;
+					const FDockTabStyle& TabStyle = CALL_PROTECTED(SDockTab, ReturnVal, GetCurrentStyle)();
+
+					bFound = ButtonStyle == &TabStyle.CloseButtonStyle;
 				}
 
 				return bFound;
@@ -1005,23 +1005,29 @@ TSharedPtr<FLogTabInfo> SLogWidget::GetActiveTabInfo() const
 }
 
 
-void SLogWidget::AddLine(ELogType LogType, TSharedRef<FString> LogLine, FSlateColor LogColor/*=FSlateColor::UseForeground()*/,
+void SLogWidget::AddLine(ELogType InLogType, TSharedRef<FString> LogLine, FSlateColor LogColor/*=FSlateColor::UseForeground()*/,
 							bool bTakeTabFocus/*=false*/)
 {
-	TSharedRef<FLogLine> CurLogEntry = MakeShareable(new FLogLine(LogType, LogLine, LogColor));
+	TSharedRef<FLogLine> CurLogEntry = MakeShareable(new FLogLine(InLogType, LogLine, LogColor));
 
 	// Add the line to the master list
 	LogLines.Add(CurLogEntry);
 
 	TSharedPtr<FLogTabInfo> ActiveTab = GetActiveTabInfo();
 
-	bool bLineInTabFocus = ActiveTab.IsValid() && !!(ActiveTab->Filter & LogType);
+	auto MatchesTabFilter =
+		[&](const TSharedPtr<FLogTabInfo>& InTab)
+		{
+			return InTab->Filter == ELogType::All || !!(InTab->Filter & InLogType);
+		};
+
+	bool bLineInTabFocus = ActiveTab.IsValid() && MatchesTabFilter(ActiveTab);
 	TSharedPtr<FLogTabInfo> FocusTab = NULL;
 
 	// Then add it to each log tab, if it passes that tabs filter
 	for (auto CurTabInfo : LogTabs)
 	{
-		if (!!(CurTabInfo->Filter & LogType))
+		if (MatchesTabFilter(CurTabInfo))
 		{
 			// If the tab is not presently open, open it now
 			if (!CurTabInfo->bTabOpen && LogTabManager.IsValid())
diff --git a/Engine/Plugins/NetcodeUnitTest/NetcodeUnitTest/Source/NetcodeUnitTest/Private/UnitTest.cpp b/Engine/Plugins/NetcodeUnitTest/NetcodeUnitTest/Source/NetcodeUnitTest/Private/UnitTest.cpp
--- a/Engine/Plugins/NetcodeUnitTest/NetcodeUnitTest/Source/NetcodeUnitTest/Private/UnitTest.cpp
+++ b/Engine/Plugins/NetcodeUnitTest/NetcodeUnitTest/Source/NetcodeUnitTest/Private/UnitTest.cpp
@@ -175,13 +175,13 @@ void UUnitTest::CleanupUnitTest()
 	}
 }
 
-void UUnitTest::NotifyLocalLog(ELogType LogType, const TCHAR* Data, ELogVerbosity::Type Verbosity, const class FName& Category)
+void UUnitTest::NotifyLocalLog(ELogType InLogType, const TCHAR* Data, ELogVerbosity::Type Verbosity, const class FName& Category)
 {
 	if (LogWindow.IsValid())
 	{
 		TSharedPtr<SLogWidget>& LogWidget = LogWindow->LogWidget;
 
-		ELogType LogOrigin = (LogType & ELogType::OriginMask);
+		ELogType LogOrigin = (InLogType & ELogType::OriginMask);
 
 		if (LogWidget.IsValid() && LogOrigin != ELogType::None && !(LogOrigin & ELogType::OriginVoid))
 		{
@@ -224,12 +224,52 @@ void UUnitTest::NotifyLocalLog(ELogType LogType, const TCHAR* Data, ELogVerbosit
 
 				bool bRequestFocus = !!(LogOrigin & ELogType::FocusMask);
 
-				LogWidget->AddLine(LogType, MakeShareable(new FString(LogLine)), CurLogColor, bRequestFocus);
+				LogWidget->AddLine(InLogType, MakeShareable(new FString(LogLine)), CurLogColor, bRequestFocus);
 			}
 		}
 	}
 }
 
+void UUnitTest::NotifyDeveloperModeRequest(bool bInDeveloperMode)
+{
+	bDeveloperMode = bInDeveloperMode;
+}
+
+bool UUnitTest::NotifyConsoleCommandRequest(FString CommandContext, FString Command)
+{
+	bool bHandled = false;
+	static TArray<FString> BadCmds = TArrayBuilder<FString>()
+		.Add("exit");
+
+	// Don't execute commands that crash
+	if (BadCmds.Contains(Command))
+	{
+		bHandled = true;
+
+		UNIT_LOG(ELogType::OriginConsole,
+					TEXT("Can't execute command '%s', it's in the 'bad commands' list (i.e. probably crashes)"), *Command);
+	}
+
+	if (!bHandled)
+	{
+		if (CommandContext == TEXT("Global"))
+		{
+			UNIT_LOG_BEGIN(this, ELogType::OriginConsole);
+			bHandled = GEngine->Exec(NULL, *Command, *GLog);
+			UNIT_LOG_END();
+		}
+	}
+
+	return bHandled;
+}
+
+void UUnitTest::GetCommandContextList(TArray<TSharedPtr<FString>>& OutList, FString& OutDefaultContext)
+{
+	OutList.Add(MakeShareable(new FString(TEXT("Global"))));
+
+	OutDefaultContext = TEXT("Global");
+}
+
 void UUnitTest::PostUnitTick(float DeltaTime)
 {
 	if (!bDeveloperMode && UnitTestTimeout > 0 && LastTimeoutReset != 0.0 && TimeoutExpire != 0.0)
@@ -256,6 +296,15 @@ void UUnitTest::PostUnitTick(float DeltaTime)
 	}
 }
 
+bool UUnitTest::IsTickable() const
+{
+	bool bReturnVal = Super::IsTickable();
+
+	bReturnVal = bReturnVal || UnitTestTimeout > 0;
+
+	return bReturnVal;
+}
+
 void UUnitTest::TickIsComplete(float DeltaTime)
 {
 	// If the unit test was verified as success/fail (i.e. as completed), end it now
@@ -295,5 +344,3 @@ void UUnitTest::TickIsComplete(float DeltaTime)
 	}
 }
 
-
-
diff --git a/Engine/Plugins/NetcodeUnitTest/NetcodeUnitTest/Source/NetcodeUnitTest/Private/UnitTestCommandlet.cpp b/Engine/Plugins/NetcodeUnitTest/NetcodeUnitTest/Source/NetcodeUnitTest/Private/UnitTestCommandlet.cpp
--- a/Engine/Plugins/NetcodeUnitTest/NetcodeUnitTest/Source/NetcodeUnitTest/Private/UnitTestCommandlet.cpp
+++ b/Engine/Plugins/NetcodeUnitTest/NetcodeUnitTest/Source/NetcodeUnitTest/Private/UnitTestCommandlet.cpp
@@ -73,6 +73,8 @@ int32 UUnitTestCommandlet::Main(const FString& Params)
 		GIsRequestingExit = true;
 	}
 
+	UE_LOG(LogUnitTest, Log, TEXT("NetcodeUnitTest built to target mainline CL '%i'."), TARGET_UE4_CL);
+
 	if (!GIsRequestingExit)
 	{
 		GIsRunning = true;
diff --git a/Engine/Plugins/NetcodeUnitTest/NetcodeUnitTest/Source/NetcodeUnitTest/Private/UnitTestEnvironment.cpp b/Engine/Plugins/NetcodeUnitTest/NetcodeUnitTest/Source/NetcodeUnitTest/Private/UnitTestEnvironment.cpp
--- a/Engine/Plugins/NetcodeUnitTest/NetcodeUnitTest/Source/NetcodeUnitTest/Private/UnitTestEnvironment.cpp
+++ b/Engine/Plugins/NetcodeUnitTest/NetcodeUnitTest/Source/NetcodeUnitTest/Private/UnitTestEnvironment.cpp
@@ -89,6 +89,10 @@ FString FUnitTestEnvironment::GetDefaultServerParameters(FString InLogCmds/*=TEX
 		ReturnVal += FString::Printf(TEXT(" -ExecCmds=\"%s\""), *FullExecCmds);
 	}
 
+	// Need to force all shader compilation, to happen with ShaderCompileWorker, so it can be detected easily;
+	// sometimes it would occur within threads in the main UE4 process, which was not possible to detect, and which this disables
+	ReturnVal += TEXT(" -ini:Engine:[DevOptions.Shaders]:bAllowAsynchronousShaderCompiling=False");
+
 	return ReturnVal;
 }
 
@@ -116,6 +120,10 @@ FString FUnitTestEnvironment::GetDefaultClientParameters()
 		ReturnVal += CmdLineClientParms;
 	}
 
+	// Need to force all shader compilation, to happen with ShaderCompileWorker, so it can be detected easily;
+	// sometimes it would occur within threads in the main UE4 process, which was not possible to detect, and which this disables
+	ReturnVal += TEXT(" -ini:Engine:[DevOptions.Shaders]:bAllowAsynchronousShaderCompiling=False");
+
 	SetupDefaultClientParameters(ReturnVal);
 
 	return ReturnVal;
@@ -147,6 +155,7 @@ void FUnitTestEnvironment::GetServerProgressLogs(const TArray<FString>*& OutStar
 		// Logs which should trigger a timeout reset for all games
 		TimeoutResetLogs.Add(TEXT("LogStaticMesh: Building static mesh "));
 		TimeoutResetLogs.Add(TEXT("LogMaterial: Missing cached shader map for material "));
+		TimeoutResetLogs.Add(TEXT("LogTexture:Display: Building textures: "));
 		TimeoutResetLogs.Add(TEXT("Dumping tracked stack traces for TraceName '"));
 		TimeoutResetLogs.Add(TEXT("Dumping once-off stack trace for TraceName '"));
 
@@ -170,7 +179,9 @@ void FUnitTestEnvironment::GetClientProgressLogs(const TArray<FString>*& OutTime
 	if (!bSetupLogs)
 	{
 		// Logs which should trigger a timeout reset for all games
+		TimeoutResetLogs.Add(TEXT("LogStaticMesh: Building static mesh "));
 		TimeoutResetLogs.Add(TEXT("LogMaterial: Missing cached shader map for material "));
+		TimeoutResetLogs.Add(TEXT("LogTexture:Display: Building textures: "));
 
 		InitializeClientProgressLogs(TimeoutResetLogs);
 
@@ -180,3 +191,20 @@ void FUnitTestEnvironment::GetClientProgressLogs(const TArray<FString>*& OutTime
 	OutTimeoutResetLogs = &TimeoutResetLogs;
 }
 
+void FUnitTestEnvironment::GetProgressBlockingProcesses(const TArray<FString>*& OutBlockingProcesses)
+{
+	static TArray<FString> BlockingProcesses;
+	static bool bSetupLogs = false;
+
+	if (!bSetupLogs)
+	{
+		BlockingProcesses.Add(TEXT("ShaderCompileWorker"));
+
+		InitializeProgressBlockingProcesses(BlockingProcesses);
+
+		bSetupLogs = true;
+	}
+
+	OutBlockingProcesses = &BlockingProcesses;
+}
+
diff --git a/Engine/Plugins/NetcodeUnitTest/NetcodeUnitTest/Source/NetcodeUnitTest/Private/UnitTestManager.cpp b/Engine/Plugins/NetcodeUnitTest/NetcodeUnitTest/Source/NetcodeUnitTest/Private/UnitTestManager.cpp
--- a/Engine/Plugins/NetcodeUnitTest/NetcodeUnitTest/Source/NetcodeUnitTest/Private/UnitTestManager.cpp
+++ b/Engine/Plugins/NetcodeUnitTest/NetcodeUnitTest/Source/NetcodeUnitTest/Private/UnitTestManager.cpp
@@ -148,7 +148,7 @@ bool UUnitTestManager::QueueUnitTest(UClass* UnitTestClass, bool bRequeued/*=fal
 
 
 	bool bValidUnitTestClass = UnitTestClass->IsChildOf(UUnitTest::StaticClass()) && UnitTestClass != UUnitTest::StaticClass() &&
-								UnitTestClass != UClientUnitTest::StaticClass();
+								UnitTestClass != UClientUnitTest::StaticClass() && UnitTestClass != UProcessUnitTest::StaticClass();
 
 	UUnitTest* UnitTestDefault = (bValidUnitTestClass ? Cast<UUnitTest>(UnitTestClass->GetDefaultObject()) : NULL);
 	bool bSupportsAllGames = (bValidUnitTestClass ? UnitTestDefault->GetSupportedGames().Contains("NullUnitEnv") : false);
@@ -533,11 +533,11 @@ void UUnitTestManager::NotifyUnitTestCleanup(UUnitTest* InUnitTest)
 {
 	ActiveUnitTests.Remove(InUnitTest);
 
-	UClientUnitTest* CurClientUnitTest = Cast<UClientUnitTest>(InUnitTest);
+	UProcessUnitTest* CurProcUnitTest = Cast<UProcessUnitTest>(InUnitTest);
 
-	if (CurClientUnitTest != NULL)
+	if (CurProcUnitTest != NULL)
 	{
-		CurClientUnitTest->OnServerSuspendState.Unbind();
+		CurProcUnitTest->OnSuspendStateChange.Unbind();
 	}
 
 
@@ -608,13 +608,14 @@ void UUnitTestManager::NotifyLogWindowClosed(const TSharedRef<SWindow>& ClosedWi
 
 		if (CurUnitTest != NULL)
 		{
-			UClientUnitTest* CurClientUnitTest = Cast<UClientUnitTest>(CurUnitTest);
+			UProcessUnitTest* CurProcUnitTest = Cast<UProcessUnitTest>(CurUnitTest);
 
-			if (CurClientUnitTest != NULL)
+			if (CurProcUnitTest != NULL)
 			{
-				CurClientUnitTest->OnServerSuspendState.Unbind();
+				CurProcUnitTest->OnSuspendStateChange.Unbind();
 			}
 
+
 			if (!CurUnitTest->bCompleted && !CurUnitTest->bAborted)
 			{
 				// Show a message box, asking the player if they'd like to also abort the unit test
@@ -1055,25 +1056,23 @@ void UUnitTestManager::OpenUnitTestLogWindow(UUnitTest* InUnitTest)
 	if (LogWindowManager != NULL)
 	{
 		InUnitTest->LogWindow = LogWindowManager->CreateLogWindow(InUnitTest->GetUnitTestName(), InUnitTest->GetExpectedLogTypes());
+		TSharedPtr<SLogWidget> CurLogWidget = (InUnitTest->LogWindow.IsValid() ? InUnitTest->LogWindow->LogWidget : NULL);
 
-		UClientUnitTest* CurClientUnitTest = Cast<UClientUnitTest>(InUnitTest);
-
-		if (InUnitTest->LogWindow.IsValid() && CurClientUnitTest != NULL)
+		if (CurLogWidget.IsValid())
 		{
-			TSharedPtr<SLogWidget> CurLogWidget = InUnitTest->LogWindow->LogWidget;
+			// Setup the widget console command context list, and then bind the console command delegate
+			InUnitTest->GetCommandContextList(CurLogWidget->ConsoleContextList, CurLogWidget->DefaultConsoleContext);
 
-			if (CurLogWidget.IsValid())
-			{
-				CurLogWidget->OnSuspendClicked.BindUObject(CurClientUnitTest, &UClientUnitTest::NotifySuspendRequest);
-				CurClientUnitTest->OnServerSuspendState.BindSP(CurLogWidget.Get(), &SLogWidget::OnSuspendStateChanged);
-
-				CurLogWidget->OnDeveloperClicked.BindUObject(CurClientUnitTest, &UClientUnitTest::NotifyDeveloperModeRequest);
+			CurLogWidget->OnConsoleCommand.BindUObject(InUnitTest, &UUnitTest::NotifyConsoleCommandRequest);
+			CurLogWidget->OnDeveloperClicked.BindUObject(InUnitTest, &UUnitTest::NotifyDeveloperModeRequest);
 
 
-				// Setup the widget console command context list, and then bind the console command delegate
-				CurClientUnitTest->GetCommandContextList(CurLogWidget->ConsoleContextList, CurLogWidget->DefaultConsoleContext);
+			UProcessUnitTest* CurProcUnitTest = Cast<UProcessUnitTest>(InUnitTest);
 
-				CurLogWidget->OnConsoleCommand.BindUObject(CurClientUnitTest, &UClientUnitTest::NotifyConsoleCommandRequest);
+			if (CurProcUnitTest != NULL)
+			{
+				CurLogWidget->OnSuspendClicked.BindUObject(CurProcUnitTest, &UProcessUnitTest::NotifySuspendRequest);
+				CurProcUnitTest->OnSuspendStateChange.BindSP(CurLogWidget.Get(), &SLogWidget::OnSuspendStateChanged);
 			}
 		}
 	}
@@ -1335,8 +1334,8 @@ bool UUnitTestManager::Exec(UWorld* InWorld, const TCHAR* Cmd, FOutputDevice& Ar
 	{
 		if (UnitTestClassDefaults[i]->GetUnitTestDate() == FDateTime::MinValue())
 		{
-			Ar.Logf(TEXT("ERROR: Unit Test '%s' does not have a date set!!!! A date must be added to every unit test!"),
-					*UnitTestClassDefaults[i]->GetUnitTestName());
+			Ar.Logf(TEXT("ERROR: Unit Test '%s' (%s) does not have a date set!!!! A date must be added to every unit test!"),
+					*UnitTestClassDefaults[i]->GetUnitTestName(), *UnitTestClassDefaults[i]->GetClass()->GetName());
 		}
 	}
 
diff --git a/Engine/Plugins/NetcodeUnitTest/NetcodeUnitTest/Source/NetcodeUnitTest/Private/UnitTests/VMReflection.cpp b/Engine/Plugins/NetcodeUnitTest/NetcodeUnitTest/Source/NetcodeUnitTest/Private/UnitTests/VMReflection.cpp
--- a/Engine/Plugins/NetcodeUnitTest/NetcodeUnitTest/Source/NetcodeUnitTest/Private/UnitTests/VMReflection.cpp
+++ b/Engine/Plugins/NetcodeUnitTest/NetcodeUnitTest/Source/NetcodeUnitTest/Private/UnitTests/VMReflection.cpp
@@ -774,7 +774,7 @@ bool UVMReflection::ExecuteUnitTest()
 		}
 
 		bool bError = false;
-		FText Result = FText::GetEmpty();// (FText)((FVMReflection(TestObj)->*"TextProp"), &bError);
+		FText Result = (FText)(FVMReflection(TestObj)->*"TextProp", &bError);
 
 		TestResults.Add(TEXT("FText Reading"), (!bError && Result.EqualTo(TargetResult)));
 	}
@@ -1364,7 +1364,7 @@ bool UVMReflection::ExecuteUnitTest()
 		bool bError = false;
 		FText Result;
 		
-		Result = FText::GetEmpty();// (FText)((FVMReflection(TestObj)->*"DynTextPropArray")["FText"][0], &bError);
+		Result = (FText)((FVMReflection(TestObj)->*"DynTextPropArray")["FText"][0], &bError);
 
 		TestResults.Add(TEXT("Dynamic Array Type Verify (FText)"), (!bError && Result.EqualTo(TargetResult)));
 	}
diff --git a/Engine/Plugins/NetcodeUnitTest/NetcodeUnitTest/Source/NetcodeUnitTest/Public/NUTUtil.h b/Engine/Plugins/NetcodeUnitTest/NetcodeUnitTest/Source/NetcodeUnitTest/Public/NUTUtil.h
--- a/Engine/Plugins/NetcodeUnitTest/NetcodeUnitTest/Source/NetcodeUnitTest/Public/NUTUtil.h
+++ b/Engine/Plugins/NetcodeUnitTest/NetcodeUnitTest/Source/NetcodeUnitTest/Public/NUTUtil.h
@@ -174,6 +174,10 @@ template struct AccessPrivate<AShooterCharacterServerEquipWeaponAccessor, &AShoo
 	\
 	template struct AccessPrivate<InClass##VarName##Accessor, &InClass::VarName>;
 
+// @todo #JohnB: Unfortunately, this broke in VS2015 for functions, as '&InClass::FuncName' gives an error;
+//					strictly speaking, this falls within the C++ standard and should work, so perhaps make a bug report
+//					See SLogWidget.cpp, for an alternative way of hack-accessing private members, using template specialization
+#if 0
 /**
  * Defines a class and template specialization, for a function, needed for use with the GET_PRIVATE hook below
  *
@@ -195,6 +199,7 @@ template struct AccessPrivate<AShooterCharacterServerEquipWeaponAccessor, &AShoo
 
 #define IMPLEMENT_GET_PRIVATE_FUNC(InClass, FuncName, FuncRet, FuncParms) \
 	IMPLEMENT_GET_PRIVATE_FUNC_CONST(InClass, FuncName, FuncRet, FuncParms, ;)
+#endif
 
 
 /**
@@ -207,8 +212,35 @@ template struct AccessPrivate<AShooterCharacterServerEquipWeaponAccessor, &AShoo
  */
 #define GET_PRIVATE(InClass, InObj, MemberName) (*InObj).*GetPrivate(InClass##MemberName##Accessor())
 
+// @todo #JohnB: Restore if fixed in VS2015
+#if 0
 // Version of above, for calling private functions
 #define CALL_PRIVATE(InClass, InObj, MemberName) (GET_PRIVATE(InClass, InObj, MemberName))
+#endif
+
+
+/**
+ * Defines a class used for hack-accessing protected functions, through the CALL_PROTECTED macro below
+ *
+ * @param InClass		The class being accessed (not a string, just the class, i.e. FStackTracker)
+ * @param FuncName		Name of the function being accessed (again, not a string)
+ * @param FuncRet		The return type of the function
+ * @param FuncParms		The function parameters
+ * @param FuncParmNames	The names of the function parameters, as passed from one function to another
+ * @param FuncModifier	(Optional) Modifier placed after the function - usually 'const'
+ */
+#define IMPLEMENT_GET_PROTECTED_FUNC_CONST(InClass, FuncName, FuncRet, FuncParms, FuncParmNames, FuncModifier) \
+	struct InClass##FuncName##Accessor : public InClass \
+	{ \
+		FORCEINLINE FuncRet FuncName##Accessor(FuncParms) FuncModifier \
+		{\
+			return FuncName(FuncParmNames); \
+		} \
+	};
+
+// Version of GET_PRIVATE, for calling protected functions
+#define CALL_PROTECTED(InClass, InObj, MemberName) ((InClass##MemberName##Accessor*)&(*InObj))->MemberName##Accessor
+
 
 
 /**
diff --git a/Engine/Plugins/NetcodeUnitTest/NetcodeUnitTest/Source/NetcodeUnitTest/Public/NetcodeUnitTest.h b/Engine/Plugins/NetcodeUnitTest/NetcodeUnitTest/Source/NetcodeUnitTest/Public/NetcodeUnitTest.h
--- a/Engine/Plugins/NetcodeUnitTest/NetcodeUnitTest/Source/NetcodeUnitTest/Public/NetcodeUnitTest.h
+++ b/Engine/Plugins/NetcodeUnitTest/NetcodeUnitTest/Source/NetcodeUnitTest/Public/NetcodeUnitTest.h
@@ -33,6 +33,7 @@
  * List of mainline UE4 engine changelists, that required a NetcodeUnitTest compatibility adjustment (newest first)
  * Every time you update NetcodeUnitTest code, to toggle based on a changelist, add a define for it with the CL number here
  */
+#define CL_FENGINEVERSION	2655102
 #define CL_INITCONNPARAM	2567692
 #define CL_CONSTUNIQUEID	2540329
 #define CL_CONSTNETCONN		2501704
@@ -55,7 +56,7 @@
  *
  * If in doubt, set to the top CL from list above ('List of mainline UE4 engine changelists'), and work your way down until it compiles.
  */
-#define TARGET_UE4_CL 2619427	// IMPORTANT: If you're hovering over this because compiling failed, you need to adjust this value.
+#define TARGET_UE4_CL 2664876	// IMPORTANT: If you're hovering over this because compiling failed, you need to adjust this value.
 
 
 /**
diff --git a/Engine/Plugins/NetcodeUnitTest/NetcodeUnitTest/Source/NetcodeUnitTest/Public/UI/SLogWidget.h b/Engine/Plugins/NetcodeUnitTest/NetcodeUnitTest/Source/NetcodeUnitTest/Public/UI/SLogWidget.h
--- a/Engine/Plugins/NetcodeUnitTest/NetcodeUnitTest/Source/NetcodeUnitTest/Public/UI/SLogWidget.h
+++ b/Engine/Plugins/NetcodeUnitTest/NetcodeUnitTest/Source/NetcodeUnitTest/Public/UI/SLogWidget.h
@@ -191,12 +191,12 @@ public:
 	/**
 	 * Adds a log line to the log window
 	 *
-	 * @param LogType			The type of log output this is (for filtering)
+	 * @param InLogType			The type of log output this is (for filtering)
 	 * @param LogLine			The log line
 	 * @param LogColor			The colour to apply to the log line
 	 * @param bTakeTabFocus		If this log line isn't displayed on the currently focused tab, switch focus to a tab that does show it
 	 */
-	void AddLine(ELogType LogType, TSharedRef<FString> LogLine, FSlateColor LogColor=FSlateColor::UseForeground(),
+	void AddLine(ELogType InLogType, TSharedRef<FString> LogLine, FSlateColor LogColor=FSlateColor::UseForeground(),
 					bool bTakeTabFocus=false);
 
 
diff --git a/Engine/Plugins/NetcodeUnitTest/NetcodeUnitTest/Source/NetcodeUnitTest/Public/UnitTestEnvironment.h b/Engine/Plugins/NetcodeUnitTest/NetcodeUnitTest/Source/NetcodeUnitTest/Public/UnitTestEnvironment.h
--- a/Engine/Plugins/NetcodeUnitTest/NetcodeUnitTest/Source/NetcodeUnitTest/Public/UnitTestEnvironment.h
+++ b/Engine/Plugins/NetcodeUnitTest/NetcodeUnitTest/Source/NetcodeUnitTest/Public/UnitTestEnvironment.h
@@ -153,6 +153,24 @@ public:
 	}
 
 
+	/**
+	 * Returns child process names, that indicate progress blockers in starting up (e.g. shader compiler)
+	 * All child processes, including server/client child processes (and their children) are checked.
+	 *
+	 * @param OutBlockingProcesses	Outputs process names, for processes which indicate a blocking task, which is blocking progression
+	 */
+	void GetProgressBlockingProcesses(const TArray<FString>*& OutBlockingProcesses);
+
+	/**
+	 * Called when initializing the static arrays, containing the progress blocking processes
+	 *
+	 * @param BlockingProcesses		The static array containing processes indicating a blocking task
+	 */
+	virtual void InitializeProgressBlockingProcesses(TArray<FString>& BlockingProcesses)
+	{
+	}
+
+
 	/**
 	 * Executes game-specific handling, upon client unit tests initializing the client PlayerController
 	 *
