Non-unity compile fixes.

diff --git a/Engine/Plugins/NetcodeUnitTest/NetcodeUnitTest/Source/NetcodeUnitTest/Private/ClientUnitTest.cpp b/Engine/Plugins/NetcodeUnitTest/NetcodeUnitTest/Source/NetcodeUnitTest/Private/ClientUnitTest.cpp
--- a/Engine/Plugins/NetcodeUnitTest/NetcodeUnitTest/Source/NetcodeUnitTest/Private/ClientUnitTest.cpp
+++ b/Engine/Plugins/NetcodeUnitTest/NetcodeUnitTest/Source/NetcodeUnitTest/Private/ClientUnitTest.cpp
@@ -17,13 +17,10 @@
 #include "NUTUtilReflection.h"
 
 #include "LogWindowManager.h"
-#include "SLogWindow.h"
-#include "SLogWidget.h"
 
 #include "OnlineBeaconClient.h"
 
 #include "Engine/ActorChannel.h"
-#include "Regex.h"
 
 
 // @todo #JohnBMultiFakeClient: Eventually, move >all< of the minimal/headless client handling code, into a new/separate class,
diff --git a/Engine/Plugins/NetcodeUnitTest/NetcodeUnitTest/Source/NetcodeUnitTest/Private/ProcessUnitTest.cpp b/Engine/Plugins/NetcodeUnitTest/NetcodeUnitTest/Source/NetcodeUnitTest/Private/ProcessUnitTest.cpp
--- a/Engine/Plugins/NetcodeUnitTest/NetcodeUnitTest/Source/NetcodeUnitTest/Private/ProcessUnitTest.cpp
+++ b/Engine/Plugins/NetcodeUnitTest/NetcodeUnitTest/Source/NetcodeUnitTest/Private/ProcessUnitTest.cpp
@@ -3,9 +3,15 @@
 #include "NetcodeUnitTestPCH.h"
 
 #include "ProcessUnitTest.h"
+#include "UnitTestManager.h"
 
 #include "UnitTestEnvironment.h"
 
+#include "SLogWindow.h"
+#include "SLogWidget.h"
+
+#include "Regex.h"
+
 
 /**
  * UProcessUnitTest
