Copying //UE4/Dev-Networking to //UE4/Dev-Main (Source: //UE4/Dev-Networking @ 3596412)

============================
  MAJOR FEATURES & CHANGES
============================

Change 3497498 by [REDACTED_PERSON]

	Added CVar for toggling NetConnection sequence randomization 'net.RandomizeSequence'.

	Tweaked asserts related to sequence randomization, to remove some false positives.

Change 3520192 by [REDACTED_PERSON]

	Added delegate hooks to the base netcode, new virtual functions, and made the UChannel classes overridable through UNetDriver, to allow netcode hooks without requiring a dependency on Online plugin packages - required to keep NetcodeUnitTest functionality, after Online packages were moved out of Engine.

Change 3520215 by [REDACTED_PERSON]

	Second pass at big NetcodeUnitTest Minimal Client refactor, consolidating most of the client helper code into MinimalClient, and refactoring the ClientUnitTest API (touching all unit tests).

	Also eliminated Online plugin dependencies (removing UnitTestNetDriver and UnitTestNetConnection), replacing them with low level netcode hooks.

Change 3526952 by [REDACTED_PERSON]

	Fixed minimal client PackageMap override.

Change 3528488 by [REDACTED_PERSON]

	Added actor channel hook, to disable NMT_ActorChannelFailure notification.

Change 3528491 by [REDACTED_PERSON]

	Fixed unit test actor channel blocking from triggering disconnects, by using an Engine hook to block NMT_ActorChannelFailure, instead of using a more complex/unreliable method of blocking.

Change 3529906 by [REDACTED_PERSON]

	Modified commandline parse function, to detect when it is incorrectly parsing a parameter, from within another parameters value (not exhaustive).

	For example, this commandline only contains the parameter ParamA. It should not be possible to parse ParamB, as it is part of ParamA's value:
	-ParamA="-ParamB=Value"

Change 3550072 by [REDACTED_PERSON]

	Added compile-time verification for unit test and minimal client configuration flags, and refactored how the flags are set - speeds up writing new unit tests, for games with excessive launch times.

Change 3565074 by [REDACTED_PERSON]

	Fix Replication Paused logic.

Change 3565093 by [REDACTED_PERSON]

	Add a 'UseDistanceBasedRelevancy' check in APawn::IsNetRelevantFor, to make it consistent with AActor.

Change 3565132 by [REDACTED_PERSON]

	Fix Actor's GetLevel to no longer assume the immediate outer is the level. Fix similar inconsistent usages of GetOuter on Actors.

Change 3573678 by [REDACTED_PERSON]

	Update actor replication based stats even when we are currently saturated (their values will be 0 anyways)

Change 3574047 by [REDACTED_PERSON]

	Added new '-FullStdOutLogOutput' commandline parameter, for capturing full log output through stdout.

	Makes up for the shortcoming of '-AllowStdOutLogVerbosity', being limited to ELogVerbosity::Log.

Change 3575518 by [REDACTED_PERSON]

	Remove actors from UNetDriver::DestroyedStartupOrDormant actors when their level is unloaded.

Change 3575911 by [REDACTED_PERSON]

	Adding option for specifying an editor server port for PIE. This allows users to change the default port used when playing net games without interferring with other instances.

Change 3575916 by [REDACTED_PERSON]

	Fixing the UI so that it doesn't respond to the same input as the camera controls by making the playback button unfocusable.

Change 3578830 by [REDACTED_PERSON]

	Merging new PlatformCrypto plugin for interfacing with encryption libraries, with a generic OpenSSL implementation and a BCrypt/CNG implementation for Xbox One.
	Added EncryptionComponent interface so that engine code can talk to encryption components without knowing their concrete type.
	New AESHandlerComponent packet handler component which uses PlatformCrypto to encrypt packets.

Change 3578855 by [REDACTED_PERSON]

	Merging hooks for encryption in the network handshake flow. NMT_Hello messages may now send an EncryptionToken as a parameter, which can be used by games to derive an encryption key used by the PacketHandler's EncryptionComponent.

Change 3578864 by [REDACTED_PERSON]

	Merged: Changed the net.UseEncryptionToken to be more useful and renamed it to net.AllowEncryption.
	This cvar, if 0, will prevent the PacketHandler from adding the configured encryption component, and prevent UPendingNetGame and AOnlineBeaconClient from filling out the EncryptionToken parameter of NMT_Hello - which prevents the extra encryption handshake connection step.

Change 3588479 by [REDACTED_PERSON]

	Merging: Encryption token/ack changes from josh.markiewicz
	- moved encryption token request/ack to delegates
	- moved FNetworkNotify to NetworkDelegates.h
	- moved connection logic out of GameInstance and back into networking code
	-- GameInstance sends an enum and the network code does the right thing based on that

	#tests PC dedicated server connections golden path and forced failures

Change 3588830 by [REDACTED_PERSON]

	New camera logic for playercontrollers so that we set the death cams correctly.

Change 3588956 by [REDACTED_PERSON]

	Quickfix for the character mesh position becoming offset due to movement smoothing being cached out.

Change 3588974 by [REDACTED_PERSON]

	Add a cvar and code to allow testing encryption functionality in ShooterGame, ShooterGame.TestEncryption (off by default). If enabled, attempts to use the PlatformCrypo and AESHandlerComponent to encrypt game traffic with a hardcoded key (for debugging/testing only, not secure!).

Change 3592922 by [REDACTED_PERSON]

	Call AActor::PreReplication before performing initial replication triggered from an RPC.

diff --git a/Engine/Plugins/Experimental/HTML5Networking/Source/HTML5Networking/Private/WebSocketNetDriver.cpp b/Engine/Plugins/Experimental/HTML5Networking/Source/HTML5Networking/Private/WebSocketNetDriver.cpp
--- a/Engine/Plugins/Experimental/HTML5Networking/Source/HTML5Networking/Private/WebSocketNetDriver.cpp
+++ b/Engine/Plugins/Experimental/HTML5Networking/Source/HTML5Networking/Private/WebSocketNetDriver.cpp
@@ -175,64 +175,73 @@ void UWebSocketNetDriver::LowLevelSend(FString Address, void* Data, int32 CountB
 
 void UWebSocketNetDriver::ProcessRemoteFunction(class AActor* Actor, UFunction* Function, void* Parameters, FOutParmRec* OutParms, FFrame* Stack, class UObject* SubObject)
 {
-	bool bIsServer = IsServer();
+#if !UE_BUILD_SHIPPING
+	bool bBlockSendRPC = false;
+
+	SendRPCDel.ExecuteIfBound(Actor, Function, Parameters, OutParms, Stack, SubObject, bBlockSendRPC);
 
-	UNetConnection* Connection = NULL;
-	if (bIsServer)
+	if (!bBlockSendRPC)
+#endif
 	{
-		if ((Function->FunctionFlags & FUNC_NetMulticast))
+		bool bIsServer = IsServer();
+
+		UNetConnection* Connection = NULL;
+		if (bIsServer)
 		{
-			// Multicast functions go to every client
-			TArray<UNetConnection*> UniqueRealConnections;
-			for (int32 i = 0; i<ClientConnections.Num(); ++i)
+			if ((Function->FunctionFlags & FUNC_NetMulticast))
 			{
-				Connection = ClientConnections[i];
-				if (Connection)
+				// Multicast functions go to every client
+				TArray<UNetConnection*> UniqueRealConnections;
+				for (int32 i = 0; i<ClientConnections.Num(); ++i)
 				{
-					// Do relevancy check if unreliable.
-					// Reliables will always go out. This is odd behavior. On one hand we wish to garuntee "reliables always get there". On the other
-					// hand, replicating a reliable to something on the other side of the map that is non relevant seems weird.
-					//
-					// Multicast reliables should probably never be used in gameplay code for actors that have relevancy checks. If they are, the
-					// rpc will go through and the channel will be closed soon after due to relevancy failing.
-
-					bool IsRelevant = true;
-					if ((Function->FunctionFlags & FUNC_NetReliable) == 0)
+					Connection = ClientConnections[i];
+					if (Connection)
 					{
-						if (Connection->ViewTarget)
+						// Do relevancy check if unreliable.
+						// Reliables will always go out. This is odd behavior. On one hand we wish to garuntee "reliables always get there". On the other
+						// hand, replicating a reliable to something on the other side of the map that is non relevant seems weird.
+						//
+						// Multicast reliables should probably never be used in gameplay code for actors that have relevancy checks. If they are, the
+						// rpc will go through and the channel will be closed soon after due to relevancy failing.
+
+						bool IsRelevant = true;
+						if ((Function->FunctionFlags & FUNC_NetReliable) == 0)
 						{
-							FNetViewer Viewer(Connection, 0.f);
-							IsRelevant = Actor->IsNetRelevantFor(Viewer.InViewer, Viewer.ViewTarget, Viewer.ViewLocation);
+							if (Connection->ViewTarget)
+							{
+								FNetViewer Viewer(Connection, 0.f);
+								IsRelevant = Actor->IsNetRelevantFor(Viewer.InViewer, Viewer.ViewTarget, Viewer.ViewLocation);
+							}
+							else
+							{
+								// No viewer for this connection(?), just let it go through.
+								UE_LOG(LogHTML5Networking, Log, TEXT("Multicast function %s called on actor %s when a connection has no Viewer"), *Function->GetName(), *Actor->GetName());
+							}
 						}
-						else
-						{
-							// No viewer for this connection(?), just let it go through.
-							UE_LOG(LogHTML5Networking, Log, TEXT("Multicast function %s called on actor %s when a connection has no Viewer"), *Function->GetName(), *Actor->GetName());
-						}
-					}
 
-					if (IsRelevant)
-					{
-						if (Connection->GetUChildConnection() != NULL)
+						if (IsRelevant)
 						{
-							Connection = ((UChildConnection*)Connection)->Parent;
-						}
+							if (Connection->GetUChildConnection() != NULL)
+							{
+								Connection = ((UChildConnection*)Connection)->Parent;
+							}
 
-						InternalProcessRemoteFunction(Actor, SubObject, Connection, Function, Parameters, OutParms, Stack, bIsServer);
+							InternalProcessRemoteFunction(Actor, SubObject, Connection, Function, Parameters, OutParms, Stack, bIsServer);
+						}
 					}
 				}
-			}
 
-			// Return here so we don't call InternalProcessRemoteFunction again at the bottom of this function
-			return;
+				// Return here so we don't call InternalProcessRemoteFunction again at the bottom of this function
+				return;
+			}
 		}
-	}
 
-	// Send function data to remote.
-	Connection = Actor->GetNetConnection();
-	if (Connection)
-	{
-		InternalProcessRemoteFunction(Actor, SubObject, Connection, Function, Parameters, OutParms, Stack, bIsServer);
+		// Send function data to remote.
+		Connection = Actor->GetNetConnection();
+		if (Connection)
+		{
+			InternalProcessRemoteFunction(Actor, SubObject, Connection, Function, Parameters, OutParms, Stack, bIsServer);
+		}
 	}
 }
 
diff --git a/Engine/Plugins/Experimental/PlatformCrypto/PlatformCrypto.uplugin b/Engine/Plugins/Experimental/PlatformCrypto/PlatformCrypto.uplugin
new file mode 100644
--- /dev/null
+++ b/Engine/Plugins/Experimental/PlatformCrypto/PlatformCrypto.uplugin
@@ -0,0 +1,54 @@
+{
+	"FileVersion" : 3,
+	"Version" : 1,
+	"VersionName" : "1.0",
+	"FriendlyName" : "Platform Cryptography Plugin",
+	"Description" : "Exposes a unified API for cryptography functionality provided by the platform, if available. Otherwise, interfaces with OpenSSL.",
+	"Category" : "Misc",
+	"CreatedBy" : "Epic Games, Inc.",
+	"CreatedByURL" : "http://epicgames.com",
+	"DocsURL" : "",
+	"MarketplaceURL" : "",
+	"SupportURL" : "",
+	"EnabledByDefault" : false,
+	"CanContainContent" : false,
+	"IsBetaVersion" : false,
+	"Installed" : false,
+	"Modules" :
+	[
+		{
+			"Name": "PlatformCrypto",
+			"Type": "Runtime",
+			"LoadingPhase": "Default",
+			"WhitelistPlatforms": [
+				"Mac",
+				"Win32",
+				"Win64",
+				"Linux",
+				"PS4",
+				"XboxOne"
+			]
+		},
+		{
+			"Name": "PlatformCryptoOpenSSL",
+			"Type": "Runtime",
+			"LoadingPhase": "Default",
+			"WhitelistPlatforms": [
+				"Mac",
+				"Win32",
+				"Win64",
+				"Linux",
+				"PS4"
+			]
+		},
+		{
+			"Name" : "PlatformCryptoBCrypt",
+			"Type" : "Runtime",
+			"LoadingPhase" : "Default",
+			"WhitelistPlatforms" :
+			[
+				"XboxOne"
+			]
+		}
+	]
+}
\ No newline at end of file
diff --git a/Engine/Plugins/Experimental/PlatformCrypto/Source/PlatformCrypto/PlatformCrypto.Build.cs b/Engine/Plugins/Experimental/PlatformCrypto/Source/PlatformCrypto/PlatformCrypto.Build.cs
new file mode 100644
--- /dev/null
+++ b/Engine/Plugins/Experimental/PlatformCrypto/Source/PlatformCrypto/PlatformCrypto.Build.cs
@@ -0,0 +1,56 @@
+// Copyright 1998-2017 Epic Games, Inc. All Rights Reserved.
+
+namespace UnrealBuildTool.Rules
+{
+	public class PlatformCrypto : ModuleRules
+	{
+		public PlatformCrypto(ReadOnlyTargetRules Target) : base(Target)
+		{
+			PrivateIncludePaths.AddRange(
+				new string[] {
+					"Experimental/PlatformCrypto/Private",
+				}
+				);
+
+			PublicDependencyModuleNames.AddRange(
+				new string[]
+				{
+					"Core",
+				}
+				);
+
+			if (Target.Platform == UnrealTargetPlatform.XboxOne)
+			{
+				PublicDependencyModuleNames.AddRange(
+					new string[]
+					{
+						"PlatformCryptoBCrypt",
+					}
+					);
+
+				PublicIncludePathModuleNames.AddRange(
+					new string[]
+					{
+						"PlatformCryptoBCrypt"
+					}
+					);
+			}
+			else
+			{
+				PublicDependencyModuleNames.AddRange(
+					new string[]
+					{
+						"PlatformCryptoOpenSSL",
+					}
+					);
+
+				PublicIncludePathModuleNames.AddRange(
+					new string[]
+					{
+						"PlatformCryptoOpenSSL"
+					}
+					);
+			}
+		}
+	}
+}
diff --git a/Engine/Plugins/Experimental/PlatformCrypto/Source/PlatformCrypto/Private/PlatformCrypto.cpp b/Engine/Plugins/Experimental/PlatformCrypto/Source/PlatformCrypto/Private/PlatformCrypto.cpp
new file mode 100644
--- /dev/null
+++ b/Engine/Plugins/Experimental/PlatformCrypto/Source/PlatformCrypto/Private/PlatformCrypto.cpp
@@ -0,0 +1,32 @@
+// Copyright 1998-2017 Epic Games, Inc. All Rights Reserved.
+
+#include "CoreMinimal.h"
+#include "Modules/ModuleManager.h"
+#include "IPlatformCrypto.h"
+
+#include "PlatformCryptoIncludes.h"
+
+class FPlatformCrypto : public IPlatformCrypto
+{
+	/** IModuleInterface implementation */
+	virtual void StartupModule() override;
+	virtual void ShutdownModule() override;
+};
+
+IMPLEMENT_MODULE( FPlatformCrypto, PlatformCrypto )
+
+
+
+void FPlatformCrypto::StartupModule()
+{
+}
+
+
+void FPlatformCrypto::ShutdownModule()
+{
+}
+
+TUniquePtr<FEncryptionContext> IPlatformCrypto::CreateContext()
+{
+	return MakeUnique<FEncryptionContext>();
+}
diff --git a/Engine/Plugins/Experimental/PlatformCrypto/Source/PlatformCrypto/Public/PlatformCryptoIncludes.h b/Engine/Plugins/Experimental/PlatformCrypto/Source/PlatformCrypto/Public/PlatformCryptoIncludes.h
new file mode 100644
--- /dev/null
+++ b/Engine/Plugins/Experimental/PlatformCrypto/Source/PlatformCrypto/Public/PlatformCryptoIncludes.h
@@ -0,0 +1,11 @@
+// Copyright 1998-2017 Epic Games, Inc. All Rights Reserved.
+
+#pragma once
+
+#include "CoreMinimal.h"
+
+#if PLATFORM_XBOXONE
+	#include "EncryptionContextBCrypt.h"
+#else
+	#include "EncryptionContextOpenSSL.h"
+#endif
diff --git a/Engine/Plugins/Experimental/PlatformCrypto/Source/PlatformCrypto/Public/PlatformCryptoTypes.h b/Engine/Plugins/Experimental/PlatformCrypto/Source/PlatformCrypto/Public/PlatformCryptoTypes.h
new file mode 100644
--- /dev/null
+++ b/Engine/Plugins/Experimental/PlatformCrypto/Source/PlatformCrypto/Public/PlatformCryptoTypes.h
@@ -0,0 +1,20 @@
+// Copyright 1998-2017 Epic Games, Inc. All Rights Reserved.
+
+#pragma once
+
+#include "CoreMinimal.h"
+#include "Stats.h"
+
+/**
+ * Generic result type for cryptographic functions.
+ */
+enum class EPlatformCryptoResult
+{
+	Success,
+	Failure
+};
+
+/**
+ * Stat group for implementations to use.
+ */
+DECLARE_STATS_GROUP(TEXT("PlatformCrypto"), STATGROUP_PlatformCrypto, STATCAT_Advanced);
diff --git a/Engine/Plugins/Experimental/PlatformCrypto/Source/PlatformCryptoOpenSSL/PlatformCryptoOpenSSL.Build.cs b/Engine/Plugins/Experimental/PlatformCrypto/Source/PlatformCryptoOpenSSL/PlatformCryptoOpenSSL.Build.cs
new file mode 100644
--- /dev/null
+++ b/Engine/Plugins/Experimental/PlatformCrypto/Source/PlatformCryptoOpenSSL/PlatformCryptoOpenSSL.Build.cs
@@ -0,0 +1,32 @@
+// Copyright 1998-2017 Epic Games, Inc. All Rights Reserved.
+
+namespace UnrealBuildTool.Rules
+{
+	public class PlatformCryptoOpenSSL : ModuleRules
+	{
+		public PlatformCryptoOpenSSL(ReadOnlyTargetRules Target) : base(Target)
+		{
+			PrivateIncludePaths.AddRange(
+				new string[] {
+					"PlatformCryptoOpenSSL/Private",
+				}
+				);
+
+			PublicDependencyModuleNames.AddRange(
+				new string[]
+				{
+					"Core",
+				}
+				);
+
+			PublicIncludePathModuleNames.AddRange(
+				new string[]
+				{
+					"PlatformCrypto"
+				}
+				);
+
+			AddEngineThirdPartyPrivateStaticDependencies(Target, "OpenSSL");
+		}
+	}
+}
diff --git a/Engine/Plugins/Experimental/PlatformCrypto/Source/PlatformCryptoOpenSSL/Private/EncryptionContextOpenSSL.cpp b/Engine/Plugins/Experimental/PlatformCrypto/Source/PlatformCryptoOpenSSL/Private/EncryptionContextOpenSSL.cpp
new file mode 100644
--- /dev/null
+++ b/Engine/Plugins/Experimental/PlatformCrypto/Source/PlatformCryptoOpenSSL/Private/EncryptionContextOpenSSL.cpp
@@ -0,0 +1,131 @@
+// Copyright 1998-2017 Epic Games, Inc. All Rights Reserved.
+
+#include "EncryptionContextOpenSSL.h"
+
+#include <openssl/evp.h>
+
+DEFINE_LOG_CATEGORY_STATIC(LogPlatformCryptoOpenSSL, Warning, All);
+
+static const int32 AES256_KeySizeInBytes = 32;
+static const int32 AES256_BlockSizeInBytes = 16;
+
+class FScopedEVPContext
+{
+public:
+	FScopedEVPContext() :
+		Context(EVP_CIPHER_CTX_new())
+	{
+	}
+
+	~FScopedEVPContext()
+	{
+		EVP_CIPHER_CTX_free(Context);
+	}
+
+	EVP_CIPHER_CTX* Get() const { return Context; }
+
+private:
+	EVP_CIPHER_CTX* Context;
+};
+
+TArray<uint8> FEncryptionContextOpenSSL::Encrypt_AES_256_ECB(const TArrayView<const uint8> Plaintext, const TArrayView<const uint8> Key, EPlatformCryptoResult& OutResult)
+{
+	DECLARE_SCOPE_CYCLE_COUNTER(TEXT("OpenSSL AES256 Encrypt"), STAT_OpenSSL_AES_Encrypt, STATGROUP_PlatformCrypto);
+
+	OutResult = EPlatformCryptoResult::Failure;
+
+	if (Key.Num() != AES256_KeySizeInBytes)
+	{
+		UE_LOG(LogPlatformCryptoOpenSSL, Warning, TEXT("FEncryptionContextOpenSSL::Encrypt_AES_256_ECB: Key size %d is not the expected size %d."), Key.Num(), AES256_KeySizeInBytes);
+		return TArray<uint8>();
+	}
+
+	FScopedEVPContext Context;
+	if (Context.Get() == nullptr)
+	{
+		UE_LOG(LogPlatformCryptoOpenSSL, Warning, TEXT("FEncryptionContextOpenSSL::Encrypt_AES_256_ECB: failed to create EVP context."));
+		return TArray<uint8>();
+	}
+
+	const int InitResult = EVP_EncryptInit_ex(Context.Get(), EVP_aes_256_ecb(), nullptr, Key.GetData(), nullptr);
+	if (InitResult != 1)
+	{
+		UE_LOG(LogPlatformCryptoOpenSSL, Warning, TEXT("FEncryptionContextOpenSSL::Encrypt_AES_256_ECB: EVP_EncryptInit_ex failed."));
+		return TArray<uint8>();
+	}
+
+	TArray<uint8> Ciphertext;
+	Ciphertext.SetNumUninitialized(Plaintext.Num() + AES256_BlockSizeInBytes); // Allow for up to a block of padding.
+
+	int OutLength = 0;
+	const int UpdateResult = EVP_EncryptUpdate(Context.Get(), Ciphertext.GetData(), &OutLength, Plaintext.GetData(), Plaintext.Num());
+	if (UpdateResult != 1)
+	{
+		UE_LOG(LogPlatformCryptoOpenSSL, Warning, TEXT("FEncryptionContextOpenSSL::Encrypt_AES_256_ECB: EVP_EncryptUpdate failed."));
+		return TArray<uint8>();
+	}
+
+	int FinalizeLength = 0;
+	const int FinalizeResult = EVP_EncryptFinal_ex(Context.Get(), Ciphertext.GetData() + OutLength, &FinalizeLength);
+	if (FinalizeResult != 1)
+	{
+		UE_LOG(LogPlatformCryptoOpenSSL, Warning, TEXT("FEncryptionContextOpenSSL::Encrypt_AES_256_ECB: EVP_EncryptFinal_ex failed."));
+		return TArray<uint8>();
+	}
+
+	Ciphertext.SetNum(OutLength + FinalizeLength);
+
+	OutResult = EPlatformCryptoResult::Success;
+	return Ciphertext;
+}
+
+TArray<uint8> FEncryptionContextOpenSSL::Decrypt_AES_256_ECB(const TArrayView<const uint8> Ciphertext, const TArrayView<const uint8> Key, EPlatformCryptoResult& OutResult)
+{
+	DECLARE_SCOPE_CYCLE_COUNTER(TEXT("OpenSSL AES256 Decrypt"), STAT_OpenSSL_AES_Decrypt, STATGROUP_PlatformCrypto);
+
+	OutResult = EPlatformCryptoResult::Failure;
+
+	if (Key.Num() != AES256_KeySizeInBytes)
+	{
+		UE_LOG(LogPlatformCryptoOpenSSL, Warning, TEXT("FEncryptionContextOpenSSL::Decrypt_AES_256_ECB: Key size %d is not the expected size %d."), Key.Num(), AES256_KeySizeInBytes);
+		return TArray<uint8>();
+	}
+
+	FScopedEVPContext Context;
+	if (Context.Get() == nullptr)
+	{
+		UE_LOG(LogPlatformCryptoOpenSSL, Warning, TEXT("FEncryptionContextOpenSSL::Decrypt_AES_256_ECB: failed to create EVP context."));
+		return TArray<uint8>();
+	}
+
+	const int InitResult = EVP_DecryptInit_ex(Context.Get(), EVP_aes_256_ecb(), nullptr, Key.GetData(), nullptr);
+	if (InitResult != 1)
+	{
+		UE_LOG(LogPlatformCryptoOpenSSL, Warning, TEXT("FEncryptionContextOpenSSL::Decrypt_AES_256_ECB: EVP_DecryptInit_ex failed."));
+		return TArray<uint8>();
+	}
+
+	TArray<uint8> Plaintext;
+	Plaintext.SetNumUninitialized(Ciphertext.Num());
+
+	int OutLength = 0;
+	const int UpdateResult = EVP_DecryptUpdate(Context.Get(), Plaintext.GetData(), &OutLength, Ciphertext.GetData(), Ciphertext.Num());
+	if (UpdateResult != 1)
+	{
+		UE_LOG(LogPlatformCryptoOpenSSL, Warning, TEXT("FEncryptionContextOpenSSL::Decrypt_AES_256_ECB: EVP_DecryptUpdate failed."));
+		return TArray<uint8>();
+	}
+
+	int FinalizeLength = 0;
+	const int FinalizeResult = EVP_DecryptFinal_ex(Context.Get(), Plaintext.GetData() + OutLength, &FinalizeLength);
+	if (FinalizeResult != 1)
+	{
+		UE_LOG(LogPlatformCryptoOpenSSL, Warning, TEXT("FEncryptionContextOpenSSL::Decrypt_AES_256_ECB: EVP_DecryptFinal_ex failed."));
+		return TArray<uint8>();
+	}
+
+	Plaintext.SetNum(OutLength + FinalizeLength);
+
+	OutResult = EPlatformCryptoResult::Success;
+	return Plaintext;
+}
diff --git a/Engine/Plugins/Experimental/PlatformCrypto/Source/PlatformCryptoOpenSSL/Private/PlatformCryptoOpenSSL.cpp b/Engine/Plugins/Experimental/PlatformCrypto/Source/PlatformCryptoOpenSSL/Private/PlatformCryptoOpenSSL.cpp
new file mode 100644
--- /dev/null
+++ b/Engine/Plugins/Experimental/PlatformCrypto/Source/PlatformCryptoOpenSSL/Private/PlatformCryptoOpenSSL.cpp
@@ -0,0 +1,32 @@
+// Copyright 1998-2017 Epic Games, Inc. All Rights Reserved.
+
+#include "CoreMinimal.h"
+#include "Modules/ModuleManager.h"
+#include "IPlatformCryptoOpenSSL.h"
+
+
+class FPlatformCryptoOpenSSL : public IPlatformCryptoOpenSSL
+{
+	/** IModuleInterface implementation */
+	virtual void StartupModule() override;
+	virtual void ShutdownModule() override;
+};
+
+IMPLEMENT_MODULE( FPlatformCryptoOpenSSL, PlatformCryptoOpenSSL )
+
+
+
+void FPlatformCryptoOpenSSL::StartupModule()
+{
+	// This code will execute after your module is loaded into memory (but after global variables are initialized, of course.)
+}
+
+
+void FPlatformCryptoOpenSSL::ShutdownModule()
+{
+	// This function may be called during shutdown to clean up your module.  For modules that support dynamic reloading,
+	// we call this function before unloading the module.
+}
+
+
+
diff --git a/Engine/Plugins/Experimental/PlatformCrypto/Source/PlatformCryptoOpenSSL/Public/EncryptionContextOpenSSL.h b/Engine/Plugins/Experimental/PlatformCrypto/Source/PlatformCryptoOpenSSL/Public/EncryptionContextOpenSSL.h
new file mode 100644
--- /dev/null
+++ b/Engine/Plugins/Experimental/PlatformCrypto/Source/PlatformCryptoOpenSSL/Public/EncryptionContextOpenSSL.h
@@ -0,0 +1,22 @@
+// Copyright 1998-2017 Epic Games, Inc. All Rights Reserved.
+
+#pragma once
+
+#include "CoreMinimal.h"
+#include "ArrayView.h"
+#include "PlatformCryptoTypes.h"
+
+/**
+ * Interface to certain cryptographic algorithms, using OpenSSL to implement them.
+ */
+class PLATFORMCRYPTOOPENSSL_API FEncryptionContextOpenSSL
+{
+
+public:
+
+	TArray<uint8> Encrypt_AES_256_ECB(const TArrayView<const uint8> Plaintext, const TArrayView<const uint8> Key, EPlatformCryptoResult& OutResult);
+
+	TArray<uint8> Decrypt_AES_256_ECB(const TArrayView<const uint8> Ciphertext, const TArrayView<const uint8> Key, EPlatformCryptoResult& OutResult);
+};
+
+typedef FEncryptionContextOpenSSL FEncryptionContext;
diff --git a/Engine/Plugins/Experimental/PlatformCrypto/Source/PlatformCryptoOpenSSL/Public/IPlatformCryptoOpenSSL.h b/Engine/Plugins/Experimental/PlatformCrypto/Source/PlatformCryptoOpenSSL/Public/IPlatformCryptoOpenSSL.h
new file mode 100644
--- /dev/null
+++ b/Engine/Plugins/Experimental/PlatformCrypto/Source/PlatformCryptoOpenSSL/Public/IPlatformCryptoOpenSSL.h
@@ -0,0 +1,39 @@
+// Copyright 1998-2017 Epic Games, Inc. All Rights Reserved.
+
+#pragma once
+
+#include "CoreMinimal.h"
+#include "Modules/ModuleInterface.h"
+#include "Modules/ModuleManager.h"
+
+
+/**
+ * Module interface for OpenSSL cryptographic functionality. Users should generally go through the FEncryptionContext API.
+ */
+class IPlatformCryptoOpenSSL : public IModuleInterface
+{
+
+public:
+
+	/**
+	 * Singleton-like access to this module's interface.  This is just for convenience!
+	 * Beware of calling this during the shutdown phase, though.  Your module might have been unloaded already.
+	 *
+	 * @return Returns singleton instance, loading the module on demand if needed
+	 */
+	static inline IPlatformCryptoOpenSSL& Get()
+	{
+		return FModuleManager::LoadModuleChecked< IPlatformCryptoOpenSSL >( "PlatformCryptoOpenSSL" );
+	}
+
+	/**
+	 * Checks to see if this module is loaded and ready.  It is only valid to call Get() if IsAvailable() returns true.
+	 *
+	 * @return True if the module is loaded and ready to use
+	 */
+	static inline bool IsAvailable()
+	{
+		return FModuleManager::Get().IsModuleLoaded( "PlatformCryptoOpenSSL" );
+	}
+};
+
diff --git a/Engine/Plugins/NetcodeUnitTest/NUTUnrealEngine4/Source/NUTUnrealEngine4/Classes/UnitTests/PacketLimitTest.h b/Engine/Plugins/NetcodeUnitTest/NUTUnrealEngine4/Source/NUTUnrealEngine4/Classes/UnitTests/PacketLimitTest.h
--- a/Engine/Plugins/NetcodeUnitTest/NUTUnrealEngine4/Source/NUTUnrealEngine4/Classes/UnitTests/PacketLimitTest.h
+++ b/Engine/Plugins/NetcodeUnitTest/NUTUnrealEngine4/Source/NUTUnrealEngine4/Classes/UnitTests/PacketLimitTest.h
@@ -3,6 +3,8 @@
 #pragma once
 
 
+// @todo #JohnB: Restore in a game-level package, eventually
+#if 0
 #include "CoreMinimal.h"
 #include "UObject/ObjectMacros.h"
 #include "ClientUnitTest.h"
@@ -61,9 +63,6 @@ protected:
 	/** The current stage of testing */
 	ELimitTestStage TestStage;
 
-	/** The last attempted send size passed to LowLevelSend (before PacketHandlers) */
-	int32 LastSendSize;
-
 	/** The last attempted send size passed to the connection socket (after PacketHandlers) */
 	int32 LastSocketSendSize;
 
@@ -75,7 +74,8 @@ public:
 
 	virtual bool ValidateUnitTestSettings(bool bCDOCheck) override;
 
-	virtual bool ConnectFakeClient(FUniqueNetIdRepl* InNetID) override;
+	// @todo #JohnB: Remove
+	//virtual bool ConnectMinimalClient(FUniqueNetIdRepl* InNetID) override;
 
 	virtual void ExecuteClientUnitTest() override;
 
@@ -91,3 +91,5 @@ protected:
 	 */
 	void NextTestStage();
 };
+
+#endif
diff --git a/Engine/Plugins/NetcodeUnitTest/NUTUnrealEngine4/Source/NUTUnrealEngine4/Classes/UnitTests/UTT61_DebugReplicateData.h b/Engine/Plugins/NetcodeUnitTest/NUTUnrealEngine4/Source/NUTUnrealEngine4/Classes/UnitTests/UTT61_DebugReplicateData.h
--- a/Engine/Plugins/NetcodeUnitTest/NUTUnrealEngine4/Source/NUTUnrealEngine4/Classes/UnitTests/UTT61_DebugReplicateData.h
+++ b/Engine/Plugins/NetcodeUnitTest/NUTUnrealEngine4/Source/NUTUnrealEngine4/Classes/UnitTests/UTT61_DebugReplicateData.h
@@ -46,7 +46,7 @@ private:
 public:
 	virtual void InitializeEnvironmentSettings() override;
 
-	virtual bool NotifyAllowNetActor(UClass* ActorClass, bool bActorChannel);
+	virtual void NotifyAllowNetActor(UClass* ActorClass, bool bActorChannel, bool& bBlockActor);
 
 	virtual void NotifyNetActor(UActorChannel* ActorChannel, AActor* Actor);
 
diff --git a/Engine/Plugins/NetcodeUnitTest/NUTUnrealEngine4/Source/NUTUnrealEngine4/NUTUnrealEngine4.Build.cs b/Engine/Plugins/NetcodeUnitTest/NUTUnrealEngine4/Source/NUTUnrealEngine4/NUTUnrealEngine4.Build.cs
--- a/Engine/Plugins/NetcodeUnitTest/NUTUnrealEngine4/Source/NUTUnrealEngine4/NUTUnrealEngine4.Build.cs
+++ b/Engine/Plugins/NetcodeUnitTest/NUTUnrealEngine4/Source/NUTUnrealEngine4/NUTUnrealEngine4.Build.cs
@@ -15,7 +15,6 @@ namespace UnrealBuildTool.Rules
 					"Core",
 					"CoreUObject",
 					"Engine",
-					"OnlineSubsystemUtils",
 					"NetcodeUnitTest"
 				}
 			);
diff --git a/Engine/Plugins/NetcodeUnitTest/NUTUnrealEngine4/Source/NUTUnrealEngine4/Private/UnitTests/FTextCrash.cpp b/Engine/Plugins/NetcodeUnitTest/NUTUnrealEngine4/Source/NUTUnrealEngine4/Private/UnitTests/FTextCrash.cpp
--- a/Engine/Plugins/NetcodeUnitTest/NUTUnrealEngine4/Source/NUTUnrealEngine4/Private/UnitTests/FTextCrash.cpp
+++ b/Engine/Plugins/NetcodeUnitTest/NUTUnrealEngine4/Source/NUTUnrealEngine4/Private/UnitTests/FTextCrash.cpp
@@ -3,6 +3,7 @@
 #include "UnitTests/FTextCrash.h"
 
 
+#include "MinimalClient.h"
 #include "NUTActor.h"
 
 #include "UnitTestEnvironment.h"
@@ -35,9 +36,9 @@ UFTextCrash::UFTextCrash(const FObjectInitializer& ObjectInitializer)
 	UnitTestTimeout = 60;
 
 
-	UnitTestFlags |= (EUnitTestFlags::LaunchServer | EUnitTestFlags::AcceptActors | EUnitTestFlags::AcceptPlayerController |
-						EUnitTestFlags::SendRPCs | EUnitTestFlags::NotifyNetActors | EUnitTestFlags::RequireNUTActor |
-						EUnitTestFlags::ExpectServerCrash | EUnitTestFlags::ExpectDisconnect);
+	SetFlags<EUnitTestFlags::LaunchServer | EUnitTestFlags::AcceptPlayerController | EUnitTestFlags::RequireNUTActor |
+				EUnitTestFlags::ExpectServerCrash | EUnitTestFlags::ExpectDisconnect,
+				EMinClientFlags::AcceptActors | EMinClientFlags::SendRPCs | EMinClientFlags::NotifyNetActors>();
 }
 
 void UFTextCrash::InitializeEnvironmentSettings()
diff --git a/Engine/Plugins/NetcodeUnitTest/NUTUnrealEngine4/Source/NUTUnrealEngine4/Private/UnitTests/PacketLimitTest.cpp b/Engine/Plugins/NetcodeUnitTest/NUTUnrealEngine4/Source/NUTUnrealEngine4/Private/UnitTests/PacketLimitTest.cpp
--- a/Engine/Plugins/NetcodeUnitTest/NUTUnrealEngine4/Source/NUTUnrealEngine4/Private/UnitTests/PacketLimitTest.cpp
+++ b/Engine/Plugins/NetcodeUnitTest/NUTUnrealEngine4/Source/NUTUnrealEngine4/Private/UnitTests/PacketLimitTest.cpp
@@ -1,12 +1,18 @@
 // Copyright 1998-2017 Epic Games, Inc. All Rights Reserved.
 
 #include "UnitTests/PacketLimitTest.h"
+
+// @todo #JohnB: Restore in a game-level package, eventually
+#if 0
 #include "EngineGlobals.h"
 #include "Engine/Engine.h"
 #include "Sockets.h"
 #include "Engine/NetConnection.h"
-#include "IpConnection.h"
 
+// @todo #JohnB: Remove when migrated
+//#include "IpConnection.h"
+
+#include "MinimalClient.h"
 #include "UnitTestEnvironment.h"
 
 #include "Net/NUTUtilNet.h"
@@ -22,7 +28,6 @@ UPacketLimitTest::UPacketLimitTest(const FObjectInitializer& ObjectInitializer)
 	: Super(ObjectInitializer)
 	, bUseOodle(false)
 	, TestStage(ELimitTestStage::LTS_LowLevel_AtLimit)
-	, LastSendSize(0)
 	, LastSocketSendSize(0)
 	, TargetSocketSendSize(0)
 {
@@ -40,8 +45,9 @@ UPacketLimitTest::UPacketLimitTest(const FObjectInitializer& ObjectInitializer)
 	UnitTestTimeout = 60;
 
 
-	UnitTestFlags |= (EUnitTestFlags::LaunchServer | EUnitTestFlags::SkipControlJoin | EUnitTestFlags::AutoReconnect |
-						EUnitTestFlags::RequirePing | EUnitTestFlags::CaptureSendRaw);
+	SetFlags<EUnitTestFlags::LaunchServer | EUnitTestFlags::AutoReconnect | EUnitTestFlags::RequirePing |
+				EUnitTestFlags::CaptureSendRaw,
+				EMinClientFlags::SkipControlJoin>();
 }
 
 void UPacketLimitTest::InitializeEnvironmentSettings()
@@ -68,7 +74,9 @@ bool UPacketLimitTest::ValidateUnitTestSettings(bool bCDOCheck)
 	return bReturnVal;
 }
 
-bool UPacketLimitTest::ConnectFakeClient(FUniqueNetIdRepl* InNetID)
+// @todo #JohnB: Remove
+#if 0
+bool UPacketLimitTest::ConnectMinimalClient(FUniqueNetIdRepl* InNetID)
 {
 	bool bReturnVal = false;
 	bool bForceEnableDefault = UUnitTestNetConnection::bForceEnableHandler;
@@ -79,7 +87,7 @@ bool UPacketLimitTest::ConnectFakeClient(FUniqueNetIdRepl* InNetID)
 		GEngine->Exec(NULL, TEXT("Oodle ForceEnable On"));
 	}
 
-	bReturnVal = Super::ConnectFakeClient(InNetID);
+	bReturnVal = Super::ConnectMinimalClient(InNetID);
 
 	if (bUseOodle)
 	{
@@ -90,6 +98,7 @@ bool UPacketLimitTest::ConnectFakeClient(FUniqueNetIdRepl* InNetID)
 
 	return bReturnVal;
 }
+#endif
 
 void UPacketLimitTest::ExecuteClientUnitTest()
 {
@@ -97,17 +106,17 @@ void UPacketLimitTest::ExecuteClientUnitTest()
 	bool bBunchSend = TestStage == ELimitTestStage::LTS_Bunch_AtLimit || TestStage == ELimitTestStage::LTS_Bunch_OverLimit;
 
 	// You can't access LowLevelSend from UNetConnection, but you can from UIpConnection, as it's exported there
-	UIpConnection* IpConn = Cast<UIpConnection>(UnitConn);
+	UIpConnection* IpConn = (MinClient != nullptr ? Cast<UIpConnection>(MinClient->GetConn()) : nullptr);
 
 	if (IpConn != nullptr)
 	{
-		int32 PacketLimit = UnitConn->MaxPacket;
-		int32 SocketLimit = UnitConn->MaxPacket;
+		int32 PacketLimit = IpConn->MaxPacket;
+		int32 SocketLimit = IpConn->MaxPacket;
 		TArray<uint8> PacketData;
 
 		if (bBunchSend)
 		{
-			int64 FreeBits = UnitConn->SendBuffer.GetMaxBits() - MAX_BUNCH_HEADER_BITS + MAX_PACKET_TRAILER_BITS;
+			int64 FreeBits = IpConn->SendBuffer.GetMaxBits() - MAX_BUNCH_HEADER_BITS + MAX_PACKET_TRAILER_BITS;
 
 			PacketLimit = FreeBits / 8;
 
@@ -148,24 +157,34 @@ void UPacketLimitTest::ExecuteClientUnitTest()
 			}
 			else if (bBunchSend)
 			{
-				UUnitTestNetConnection* UnitTestConn = CastChecked<UUnitTestNetConnection>(UnitConn);
+				UUnitTestNetConnection* UnitTestConn = CastChecked<UUnitTestNetConnection>(IpConn);
 
 				// If the bunch is to go over the limit, disable asserts
 				bool bBunchOverLimit = TestStage == ELimitTestStage::LTS_Bunch_OverLimit;
 
 				UnitTestConn->bDisableValidateSend = bBunchOverLimit;
 
-				UnitConn->FlushNet();
+				// @todo #JohnB: To re-enable this code above, you have to reimplement this now-removed UUnitTestNetConnection code:
+#if 0
+void UUnitTestNetConnection::ValidateSendBuffer()
+{
+	if (!bDisableValidateSend)
+	{
+		Super::ValidateSendBuffer();
+	}
+}
+#endif
 
+				IpConn->FlushNet();
 
-				int32 DummyControlBunchSequence = 0;
-				FOutBunch* TestBunch = NUTNet::CreateChannelBunch(DummyControlBunchSequence, UnitConn, CHTYPE_Control, 0);
+
+				FOutBunch* TestBunch = MinClient->CreateChannelBunch(CHTYPE_Control, 0);
 
 				if (TestBunch != nullptr)
 				{
 					TestBunch->Serialize(PacketData.GetData(), PacketData.Num());
 
-					UnitConn->SendRawBunch(*TestBunch, false);
+					IpConn->SendRawBunch(*TestBunch, false);
 				}
 				else
 				{
@@ -178,7 +197,7 @@ void UPacketLimitTest::ExecuteClientUnitTest()
 				if (bBunchOverLimit)
 				{
 					// For a successful test, the bunch must cause a send error
-					if (UnitConn->SendBuffer.IsError())
+					if (IpConn->SendBuffer.IsError())
 					{
 						bPacketAtLimit = true;
 
@@ -195,7 +214,7 @@ void UPacketLimitTest::ExecuteClientUnitTest()
 				}
 
 
-				UnitConn->FlushNet();
+				IpConn->FlushNet();
 
 				UnitTestConn->bDisableValidateSend = false;
 			}
@@ -238,11 +257,6 @@ void UPacketLimitTest::ExecuteClientUnitTest()
 	}
 }
 
-void UPacketLimitTest::NotifySendRawPacket(void* Data, int32 Count, bool& bBlockSend)
-{
-	LastSendSize = Count;
-}
-
 void UPacketLimitTest::NotifySocketSendRawPacket(void* Data, int32 Count, bool& bBlockSend)
 {
 	LastSocketSendSize = Count;
@@ -258,6 +272,8 @@ void UPacketLimitTest::NotifySocketSendRawPacket(void* Data, int32 Count, bool&
 			bBlockSend = true;
 		}
 	}
+
+	Super::NotifySocketSendRawPacket(Data, Count, bBlockSend);
 }
 
 void UPacketLimitTest::NotifyProcessLog(TWeakPtr<FUnitTestProcess> InProcess, const TArray<FString>& InLogLines)
@@ -342,4 +358,4 @@ void UPacketLimitTest::NextTestStage()
 		TriggerAutoReconnect();
 	}
 }
-
+#endif
diff --git a/Engine/Plugins/NetcodeUnitTest/NUTUnrealEngine4/Source/NUTUnrealEngine4/Private/UnitTests/UTT61_DebugReplicateData.cpp b/Engine/Plugins/NetcodeUnitTest/NUTUnrealEngine4/Source/NUTUnrealEngine4/Private/UnitTests/UTT61_DebugReplicateData.cpp
--- a/Engine/Plugins/NetcodeUnitTest/NUTUnrealEngine4/Source/NUTUnrealEngine4/Private/UnitTests/UTT61_DebugReplicateData.cpp
+++ b/Engine/Plugins/NetcodeUnitTest/NUTUnrealEngine4/Source/NUTUnrealEngine4/Private/UnitTests/UTT61_DebugReplicateData.cpp
@@ -5,7 +5,7 @@
 
 
 #include "NUTActor.h"
-
+#include "MinimalClient.h"
 #include "Net/NUTUtilNet.h"
 #include "UnitTestEnvironment.h"
 
@@ -19,7 +19,7 @@ UClass* UUTT61_DebugReplicateData::RepClass = FindObject<UClass>(ANY_PACKAGE, TE
 
 UUTT61_DebugReplicateData::UUTT61_DebugReplicateData(const FObjectInitializer& ObjectInitializer)
 	: Super(ObjectInitializer)
-	, Replicator(NULL)
+	, Replicator(nullptr)
 {
 	UnitTestName = TEXT("ReplicateDataCheck");
 	UnitTestType = TEXT("DevExploit");
@@ -43,9 +43,9 @@ UUTT61_DebugReplicateData::UUTT61_DebugReplicateData(const FObjectInitializer& O
 	UnitTestTimeout = 60;
 
 
-	UnitTestFlags |= (EUnitTestFlags::LaunchServer | EUnitTestFlags::AcceptActors | EUnitTestFlags::NotifyNetActors |
-						EUnitTestFlags::AcceptPlayerController | EUnitTestFlags::RequirePlayerController | EUnitTestFlags::SendRPCs |
-						EUnitTestFlags::ExpectServerCrash | EUnitTestFlags::ExpectDisconnect);
+	SetFlags<EUnitTestFlags::LaunchServer | EUnitTestFlags::AcceptPlayerController | EUnitTestFlags::RequirePlayerController |
+				EUnitTestFlags::ExpectServerCrash | EUnitTestFlags::ExpectDisconnect,
+				EMinClientFlags::AcceptActors | EMinClientFlags::NotifyNetActors | EMinClientFlags::SendRPCs>();
 }
 
 void UUTT61_DebugReplicateData::InitializeEnvironmentSettings()
@@ -54,17 +54,17 @@ void UUTT61_DebugReplicateData::InitializeEnvironmentSettings()
 	BaseServerParameters = UnitEnv->GetDefaultServerParameters();
 }
 
-bool UUTT61_DebugReplicateData::NotifyAllowNetActor(UClass* ActorClass, bool bActorChannel)
+void UUTT61_DebugReplicateData::NotifyAllowNetActor(UClass* ActorClass, bool bActorChannel, bool& bBlockActor)
 {
-	bool bAllow = Super::NotifyAllowNetActor(ActorClass, bActorChannel);
+	Super::NotifyAllowNetActor(ActorClass, bActorChannel, bBlockActor);
 
-	if (!bAllow)
+	if (bBlockActor)
 	{
-		if (RepClass != NULL)
+		if (RepClass != nullptr)
 		{
 			if (ActorClass->IsChildOf(RepClass))
 			{
-				bAllow = true;
+				bBlockActor = false;
 			}
 		}
 		else
@@ -75,23 +75,21 @@ bool UUTT61_DebugReplicateData::NotifyAllowNetActor(UClass* ActorClass, bool bAc
 			VerificationState = EUnitTestVerification::VerifiedNeedsUpdate;
 		}
 	}
-
-	return bAllow;
 }
 
 void UUTT61_DebugReplicateData::NotifyNetActor(UActorChannel* ActorChannel, AActor* Actor)
 {
 	Super::NotifyNetActor(ActorChannel, Actor);
 
-	if (Replicator == NULL)
+	if (Replicator == nullptr)
 	{
-		if (RepClass != NULL)
+		if (RepClass != nullptr)
 		{
 			if (Actor->IsA(RepClass))
 			{
 				Replicator = Actor;
 
-				if (Replicator != NULL)
+				if (Replicator != nullptr)
 				{
 					// Once the replicator is found, pass back to the main exploit function
 					ExecuteClientUnitTest();
@@ -118,21 +116,11 @@ void UUTT61_DebugReplicateData::ExecuteClientUnitTest()
 		ResetTimeout(LogMsg);
 		UNIT_LOG(ELogType::StatusImportant, TEXT("%s"), *LogMsg);
 
-		FOutBunch* ControlChanBunch = NUTNet::CreateChannelBunch(ControlBunchSequence, UnitConn, CHTYPE_Control, 0);
-
-		if (ControlChanBunch != nullptr)
-		{
-			uint8 ControlMsg = NMT_NUTControl;
-			uint8 ControlCmd = ENUTControlCommand::Summon;
-			FString Cmd = TEXT("GameplayDebugger.GameplayDebuggingReplicator -ForceBeginPlay -GameplayDebuggerHack");
 
-			*ControlChanBunch << ControlMsg;
-			*ControlChanBunch << ControlCmd;
-			*ControlChanBunch << Cmd;
+		FString Cmd = TEXT("GameplayDebugger.GameplayDebuggingReplicator -ForceBeginPlay -GameplayDebuggerHack");
+		bool bSuccess = SendNUTControl(ENUTControlCommand::Summon, Cmd);
 
-			NUTNet::SendControlBunch(UnitConn, *ControlChanBunch);
-		}
-		else
+		if (!bSuccess)
 		{
 			UNIT_LOG(ELogType::StatusFailure, TEXT("Failed to send summon command - marking unit test as needing update."));
 
@@ -155,7 +143,7 @@ void UUTT61_DebugReplicateData::ExecuteClientUnitTest()
 		};
 
 		ServerReplicateMessageParms Parms;
-		Parms.Actor = NULL;
+		Parms.Actor = nullptr;
 		Parms.InMessage = 4; //EDebugComponentMessage::ActivateDataView;
 		Parms.DataView = -1;
 
diff --git a/Engine/Plugins/NetcodeUnitTest/NetcodeUnitTest/NetcodeUnitTest.uplugin b/Engine/Plugins/NetcodeUnitTest/NetcodeUnitTest/NetcodeUnitTest.uplugin
--- a/Engine/Plugins/NetcodeUnitTest/NetcodeUnitTest/NetcodeUnitTest.uplugin
+++ b/Engine/Plugins/NetcodeUnitTest/NetcodeUnitTest/NetcodeUnitTest.uplugin
@@ -27,14 +27,6 @@
 				"Linux"
 			]
 		}
-	],
-
-	"Plugins":
-	[
-		{
-			"Name": "OnlineSubsystemUtils",
-			"Enabled": true
-		}
 	]
 }
 
diff --git a/Engine/Plugins/NetcodeUnitTest/NetcodeUnitTest/Source/NetcodeUnitTest/Classes/ClientUnitTest.h b/Engine/Plugins/NetcodeUnitTest/NetcodeUnitTest/Source/NetcodeUnitTest/Classes/ClientUnitTest.h
--- a/Engine/Plugins/NetcodeUnitTest/NetcodeUnitTest/Source/NetcodeUnitTest/Classes/ClientUnitTest.h
+++ b/Engine/Plugins/NetcodeUnitTest/NetcodeUnitTest/Source/NetcodeUnitTest/Classes/ClientUnitTest.h
@@ -2,12 +2,11 @@
 
 #pragma once
 
-#include "CoreMinimal.h"
-#include "Misc/EnumClassFlags.h"
 #include "UObject/ObjectMacros.h"
 #include "Engine/EngineBaseTypes.h"
 #include "NetcodeUnitTest.h"
 #include "ProcessUnitTest.h"
+#include "NUTEnum.h"
 
 #include "ClientUnitTest.generated.h"
 
@@ -27,138 +26,15 @@ class UNetDriver;
 struct FFrame;
 struct FOutParmRec;
 struct FUniqueNetIdRepl;
+enum class ENUTControlCommand : uint8;
 
-// Enum definitions
-
-/**
- * Flags for configuring how individual unit tests make use of the base client unit test framework.
- * NOTE: These are flags instead of bools, partly to group these settings together, as being essential for configuring each unit test
- *
- * Types of things these flags control:
- *	- Types of remote data which are accepted/denied (channel types, actors, RPC's)
- *		- IMPORTANT: This includes local setup of e.g. actor channels, and possibly execution of RPC's in local context,
- *			which risks undefined behaviour (which is why it is disabled by default - you have to know what you're doing)
- *	- The required prerequisites needed before executing the main unit-test/payload (need a valid PlayerController? A particular actor?)
- *	- Whether or not a server is automatically launched, and if so, whether or not to capture its log output
- *	- Enabling other miscellaneous events, such as capturing raw packet data
- *
- * WARNING: This enum is not compatible with UENUM, as it is uint32, and UENUM requires uint8
- */
-enum class EUnitTestFlags : uint32 // NOTE: If you change from uint32, you need to modify the 32-bit-dependent handling in some code
-{
-	None					= 0x00000000,	// No flags
-
-	/** Sub-process flags */
-	LaunchServer			= 0x00000001,	// Whether or not to automatically launch a game server, for the unit test
-	LaunchClient			= 0x00000002,	// Whether or not to automatically launch a full game client, which connects to the server
-
-	/** Minimal-client netcode functionality */
-	AcceptActors			= 0x00000004,	// Whether or not to accept actor channels (acts as whitelist-only with NotifyAllowNetActor)
-	AcceptPlayerController	= 0x00000008,	// Whether or not to accept PlayerController creation
-	AcceptRPCs				= 0x00000010,	// Whether or not to accept execution of any actor RPC's (they are all blocked by default)
-	SendRPCs				= 0x00000020,	// Whether or not to allow RPC sending (NOT blocked by default, @todo JohnB: Add send hook)
-	SkipControlJoin			= 0x00000040,	// Whether or not to skip sending NMT_Join upon connect (or NMT_BeaconJoin for beacons)
-	BeaconConnect			= 0x00000080,	// Whether or not to connect to the servers beacon (greatly limits the connection)
-	AutoReconnect			= 0x00000100,	// Whether or not to auto-reconnect on server disconnect (NOTE: Won't catch all disconnects)
-
-	/** Unit test state-setup/requirements/prerequisites */
-	RequirePlayerController	= 0x00000200,	// Whether or not to wait for the PlayerController, before triggering ExecuteClientUnitTest
-	RequirePawn				= 0x00000400,	// Whether or not to wait for PlayerController's pawn, before ExecuteClientUnitTest
-	RequirePlayerState		= 0x00000800,	// Whether or not to wait for PlayerController's PlayerState, before ExecuteClientUnitTest
-	RequirePing				= 0x00001000,	// Whether or not to wait for a ping round-trip, before triggering ExecuteClientUnitTest
-	RequireNUTActor			= 0x00002000,	// Whether or not to wait for the NUTActor, before triggering ExecuteClientUnitTest
-	RequireBeacon			= 0x00004000,	// Whether or not to wait for beacon replication, before triggering ExecuteClientUnitTest
-	RequireMCP				= 0x00008000,	// Whether or not an MCP connection is required, before triggering ExecuteClientUnitTest
-	RequireCustom			= 0x00010000,	// Whether or not ExecuteClientUnitTest will be executed manually, within the unit test
-
-	RequirementsMask		= RequirePlayerController | RequirePawn | RequirePlayerState | RequirePing | RequireNUTActor |
-								RequireBeacon | RequireMCP | RequireCustom,
-
-	/** Unit test error/crash detection */
-	ExpectServerCrash		= 0x00020000,	// Whether or not this unit test will intentionally crash the server
-	ExpectDisconnect		= 0x00040000,	// Whether or not this unit test will intentionally trigger a disconnect from the server
-
-	/** Unit test error/crash detection debugging (NOTE: Don't use these in finalized unit tests, unit tests must handle all errors) */
-	IgnoreServerCrash		= 0x00080000,	// Whether or not server crashes should be treated as a unit test failure
-	IgnoreClientCrash		= 0x00100000,	// Whether or not client crashes should be treated as a unit test failure
-	IgnoreDisconnect		= 0x00200000,	// Whether or not minimal/fake client disconnects, should be treated as a unit test failure
-
-	/** Unit test events */
-	NotifyNetActors			= 0x00400000,	// Whether or not to trigger a 'NotifyNetActor' event, AFTER creation of actor channel actor
-	NotifyProcessEvent		= 0x00800000,	// Whether or not to trigger 'NotifyScriptProcessEvent' for every executed local function
-
-	/** Debugging */
-	CaptureReceivedRaw		= 0x01000000,	// Whether or not to capture raw (clientside) packet receives
-	CaptureSendRaw			= 0x02000000,	// Whether or not to capture raw (clientside) packet sends
-	DumpReceivedRaw			= 0x04000000,	// Whether or not to also hex-dump the raw packet receives to the log/log-window
-	DumpSendRaw				= 0x08000000,	// Whether or not to also hex-dump the raw packet sends to the log/log-window
-	DumpControlMessages		= 0x10000000,	// Whether or not to dump control channel messages, and their raw hex content
-	DumpReceivedRPC			= 0x20000000,	// Whether or not to dump RPC receives (with LogNetTraffic, detects ProcessEvent RPC fail)
-	DumpSendRPC				= 0x40000000	// Whether or not to dump RPC sends
-};
-
-// Required for bitwise operations with the above enum
-ENUM_CLASS_FLAGS(EUnitTestFlags);
-
-/**
- * Used to get name values for the above enum
- */
-inline FString GetUnitTestFlagName(EUnitTestFlags Flag)
-{
-	#define EUTF_CASE(x) case EUnitTestFlags::x : return TEXT(#x)
-
-	switch (Flag)
-	{
-		EUTF_CASE(None);
-		EUTF_CASE(LaunchServer);
-		EUTF_CASE(LaunchClient);
-		EUTF_CASE(AcceptActors);
-		EUTF_CASE(AcceptPlayerController);
-		EUTF_CASE(AcceptRPCs);
-		EUTF_CASE(SendRPCs);
-		EUTF_CASE(SkipControlJoin);
-		EUTF_CASE(BeaconConnect);
-		EUTF_CASE(RequirePlayerController);
-		EUTF_CASE(RequirePawn);
-		EUTF_CASE(RequirePlayerState);
-		EUTF_CASE(RequirePing);
-		EUTF_CASE(RequireNUTActor);
-		EUTF_CASE(RequireBeacon);
-		EUTF_CASE(RequireMCP);
-		EUTF_CASE(RequireCustom);
-		EUTF_CASE(ExpectServerCrash);
-		EUTF_CASE(ExpectDisconnect);
-		EUTF_CASE(IgnoreServerCrash);
-		EUTF_CASE(IgnoreClientCrash);
-		EUTF_CASE(IgnoreDisconnect);
-		EUTF_CASE(NotifyNetActors);
-		EUTF_CASE(NotifyProcessEvent);
-		EUTF_CASE(CaptureReceivedRaw);
-		EUTF_CASE(CaptureSendRaw);
-		EUTF_CASE(DumpReceivedRaw);
-		EUTF_CASE(DumpSendRaw);
-		EUTF_CASE(DumpControlMessages);
-		EUTF_CASE(DumpReceivedRPC);
-		EUTF_CASE(DumpSendRPC);
-
-	default:
-		if (FMath::IsPowerOfTwo((uint32)Flag))
-		{
-			return FString::Printf(TEXT("Unknown 0x%08X"), (uint32)Flag);
-		}
-		else
-		{
-			return FString::Printf(TEXT("Bad/Multiple flags 0x%08X"), (uint32) Flag);
-		}
-	}
-
-	#undef EUTF_CASE
-}
 
 
 /**
- * Base class for all unit tests depending upon a (fake/minimal) client connecting to a server.
- * Handles creation/cleanup of an entire new UWorld, UNetDriver and UNetConnection, for fast sequential unit testing.
+ * Base class for all unit tests depending upon a MinimalClient connecting to a server.
+ * The MinimalClient handles creation/cleanup of an entire new UWorld, UNetDriver and UNetConnection, for fast unit testing.
+ *
+ * NOTE: See NUTEnum.h, for important flags for configuring unit tests and the minimal client.
  * 
  * In subclasses, implement the unit test within the ExecuteClientUnitTest function (remembering to call parent)
  */
@@ -167,13 +43,11 @@ class NETCODEUNITTEST_API UClientUnitTest : public UProcessUnitTest
 {
 	GENERATED_UCLASS_BODY()
 
-
 	friend class FScopedLog;
 	friend class FScopedNetObjectReplace;
+	friend class UUnitTestActorChannel;
 	friend class UUnitTestManager;
 	friend class FUnitTestEnvironment;
-
-	// @todo #JohnB: Remove once deprecated ControlBunchSequence is removed
 	friend class UMinimalClient;
 
 
@@ -182,6 +56,9 @@ protected:
 	/** All of the internal unit test parameters/flags, for controlling state and execution */
 	EUnitTestFlags UnitTestFlags;
 
+	/** Flags for configuring the minimal client - lots of interdependencies between these and UnitTestFlags */
+	EMinClientFlags MinClientFlags;
+
 
 	/** The base URL the server should start with */
 	FString BaseServerURL;
@@ -202,7 +79,7 @@ protected:
 	/** Actors the server is allowed replicate to client (requires AllowActors flag). Use NotifyAllowNetActor for conditional allows. */
 	TArray<UClass*> AllowedClientActors;
 
-	/** Clientside RPC's that should be allowed to execute (requires NotifyProcessEvent flag; other flags also allow specific RPC's) */
+	/** Clientside RPC's that should be allowed to execute (requires minimal client NotifyProcessNetEvent flag) */
 	TArray<FString> AllowedClientRPCs;
 
 
@@ -237,15 +114,6 @@ protected:
 	UPROPERTY()
 	UMinimalClient* MinClient;
 
-	// @todo #JohnB: Deprecate this, after migrating to the minimal client code
-	//					(or perhaps keep cached? as it's likely to be used a lot)
-#if 1
-	/** Stores a reference to the created unit test net driver, for execution and later cleanup (always a UUnitTestNetDriver) */
-	UNetDriver* UnitNetDriver;
-
-	/** Stores a reference to the server connection (always a 'UUnitTestNetConnection') */
-	UNetConnection* UnitConn;
-#endif
 
 	/** Whether or not the initial connect of the fake client was triggered */
 	bool bTriggerredInitialConnect;
@@ -266,20 +134,11 @@ protected:
 	bool bUnitNUTActorSetup;
 
 	/** If EUnitTestFlags::RequireBeacon is set, stores a reference to the replicated beacon */
-	TWeakObjectPtr<class AOnlineBeaconClient> UnitBeacon;
+	TWeakObjectPtr<AActor> UnitBeacon;
 
 	/** If EUnitTestFlags::RequirePing is true, whether or not we have already received the pong */
 	bool bReceivedPong;
 
-	// @todo #JohnB: Remove after deprecated
-#if 0
-	/** For the unit test control channel, this tracks the current bunch sequence */
-	int32 ControlBunchSequence;
-#endif
-
-	// @todo #JohnB: Remove once fully migrated to minimal client
-	int32* ControlBunchSequence;
-
 	/** An expected network failure occurred, which will be handled during the next tick instead of immediately */
 	bool bPendingNetworkFailure;
 
@@ -287,6 +146,14 @@ protected:
 	bool bDetectedMCPOnline;
 
 
+	/** Whether or not a bunch was successfully sent */
+	bool bSentBunch;
+
+private:
+	/** Static reference to the OnlineBeaconClient static class */
+	static UClass* OnlineBeaconClass;
+
+
 	/**
 	 * Interface and hooked events for client unit tests
 	 */
@@ -307,10 +174,10 @@ public:
 	/**
 	 * Override this, to receive notification of NMT_NUTControl messages, from the server
 	 *
-	 * @param CmdType		The ENUTControlCommand type
+	 * @param CmdType		The command type
 	 * @param Command		The command being received
 	 */
-	virtual void NotifyNUTControl(uint8 CmdType, FString Command)
+	virtual void NotifyNUTControl(ENUTControlCommand CmdType, FString Command)
 	{
 	}
 
@@ -331,13 +198,13 @@ public:
 	virtual void NotifyHandleClientPlayer(APlayerController* PC, UNetConnection* Connection);
 
 	/**
-	 * Override this, to receive notification BEFORE a replicated actor has been created (allowing you to block, based on class)
+	 * Notification triggered BEFORE a replicated actor has been created (allowing you to block creation, based on class)
 	 *
 	 * @param ActorClass	The actor class that is about to be created
 	 * @param bActorChannel	Whether or not this actor is being created within an actor channel
-	 * @return				Whether or not to allow creation of that actor
+	 * @param bBlockActor	Whether or not to block creation of that actor (defaults to true)
 	 */
-	virtual bool NotifyAllowNetActor(UClass* ActorClass, bool bActorChannel);
+	virtual void NotifyAllowNetActor(UClass* ActorClass, bool bActorChannel, bool& bBlockActor);
 
 	/**
 	 * Override this, to receive notification AFTER an actor channel actor has been created
@@ -368,30 +235,16 @@ public:
 	}
 
 	/**
-	 * If EUnitTestFlags::CaptureSendRaw is set, this is triggered for every packet sent to the server
+	 * Triggered for every packet sent to the server, when LowLevelSend is called.
+	 * IMPORTANT: This occurs AFTER PacketHandler's have had a chance to modify packet data
+	 *
 	 * NOTE: Don't consider data safe to modify (will need to modify the implementation, if that is desired)
 	 *
 	 * @param Data			The raw data/packet being sent
 	 * @param Count			The amount of data being sent
 	 * @param bBlockSend	Whether or not to block the send (defaults to false)
 	 */
-	virtual void NotifySendRawPacket(void* Data, int32 Count, bool& bBlockSend)
-	{
-	}
-
-	/**
-	 * If EUnitTestFlags::CaptureSendRaw is set, this is triggered for every packet sent to the server
-	 * IMPORTANT: This is like the above function, except this occurs AFTER PacketHandler's have had a chance to modify packet data
-	 *
-	 * NOTE: Don't consider data safe to modify (will need to modify the implementation, if that is desired)
-	 *
-	 * @param Data		The raw data/packet being sent
-	 * @param Count		The amount of data being sent
-	 * @param bBlockSend	Whether or not to block the send (defaults to false)
-	 */
-	virtual void NotifySocketSendRawPacket(void* Data, int32 Count, bool& bBlockSend)
-	{
-	}
+	virtual void NotifySocketSendRawPacket(void* Data, int32 Count, bool& bBlockSend);
 
 	/**
 	 * Bunches received on the control channel. These need to be parsed manually,
@@ -402,31 +255,41 @@ public:
 	virtual void ReceivedControlBunch(FInBunch& Bunch);
 
 
-#if !UE_BUILD_SHIPPING
 	/**
-	 * Overridable in subclasses - can be used to control/block ALL script events
+	 * Overridable in subclasses - can be used to control/block any script events, other than receiving of RPC's (see NotifyReceiveRPC)
 	 *
 	 * @param Actor			The actor the event is being executed on
 	 * @param Function		The script function being executed
 	 * @param Parameters	The raw unparsed parameters, being passed into the function
-	 * @return				Whether or not to block the event from executing
+	 * @param bBlockEvent	Whether or not to block the event from executing
 	 */
-	virtual bool NotifyScriptProcessEvent(AActor* Actor, UFunction* Function, void* Parameters);
-#endif
+	virtual void NotifyProcessEvent(AActor* Actor, UFunction* Function, void* Parameters, bool& bBlockEvent)
+	{
+	}
 
 	/**
-	 * Overridable in subclasses - can be used to control/block sending of RPC's
+	 * Overridable in subclasses - can be used to control/block receiving of RPC's
 	 *
-	 * @param bAllowRPC		Whether or not to allow sending of the RPC
-	 * @param Actor			The actor the RPC will be called in
-	 * @param Function		The RPC to call
-	 * @param Parameters	The parameters data blob
-	 * @param OutParms		Out parameter information (irrelevant for RPC's)
-	 * @param Stack			The script stack
-	 * @param SubObject		The sub-object the RPC is being called in (if applicable)
+	 * @param Actor			The actor the RPC is being executed on
+	 * @param Function		The RPC being executed
+	 * @param Parameters	The raw unparsed parameters, being passed into the function
+	 * @param bBlockRPC		Whether or not to block the RPC from executing
 	 */
-	virtual void NotifySendRPC(bool& bAllowwRPC, AActor* Actor, UFunction* Function, void* Parameters, FOutParmRec* OutParms,
-								FFrame* Stack, UObject* SubObject)
+	virtual void NotifyReceiveRPC(AActor* Actor, UFunction* Function, void* Parameters, bool& bBlockRPC);
+
+	/**
+	 * Overridable in subclasses - can be used to control/block sending of RPC's
+	 *
+	 * @param Actor				The actor the RPC will be called in
+	 * @param Function			The RPC to call
+	 * @param Parameters		The parameters data blob
+	 * @param OutParms			Out parameter information (irrelevant for RPC's)
+	 * @param Stack				The script stack
+	 * @param SubObject			The sub-object the RPC is being called in (if applicable)
+	 * @param bBlockSendRPC		Whether or not to allow sending of the RPC
+	 */
+	virtual void NotifySendRPC(AActor* Actor, UFunction* Function, void* Parameters, FOutParmRec* OutParms, FFrame* Stack,
+								UObject* SubObject, bool& bBlockSendRPC)
 	{
 	}
 
@@ -449,6 +312,16 @@ public:
 	 * Utility functions for use by subclasses
 	 */
 public:
+	/**
+	 * Sends an NMT_NUTControl control channel message, for the server NUTActor. See NUTActor.h
+	 *
+	 * @param CommandType	The type of NUTActor control channel command being sent.
+	 * @param Command		The command parameters.
+	 * @return				Whether or not the command was sent successfully.
+	 */
+	bool SendNUTControl(ENUTControlCommand CommandType, FString Command);
+
+
 	/**
 	 * Sends the specified RPC for the specified actor, and verifies that the RPC was sent (triggering a unit test failure if not)
 	 *
@@ -461,30 +334,6 @@ public:
 	 */
 	bool SendRPCChecked(UObject* Target, const TCHAR* FunctionName, void* Parms, int16 ParmsSize, int16 ParmsSizeCorrection=0);
 
-	// @todo #JohnB: Consider deprecating/removing this, as all bits of code that use this, just implement a UnitTestServer_ function;
-	//					replace with SendUnitRPCChecked
-#if 1
-	/**
-	 * As above, except the RPC is called within a lambda
-	 * NOTE: Only call a single RPC within the lambda
-	 *
-	 * @param RPCName	The name of the RPC, for logging purposes
-	 * @param RPCCall	The lambda which executes the RPC
-	 * @return			Whether or not the RPC was sent successfully
-	 */
-	template<typename Predicate>
-	FORCEINLINE bool SendRPCChecked(FString RPCName, Predicate RPCCall)
-	{
-		bool bSuccess = false;
-
-		PreSendRPC();
-		RPCCall();
-		bSuccess = PostSendRPC(RPCName);
-
-		return bSuccess;
-	}
-#endif
-
 	/**
 	 * As above, except optimized for use with reflection
 	 *
@@ -565,6 +414,140 @@ public:
 	 * Internal base implementation and utility functions for client unit tests
 	 */
 protected:
+	/**
+	 * Validates, both at compile time (template params) or at runtime (function params), that the specified flags are valid.
+	 *
+	 * When specifying the whole flag list in one go, do a compile time check using the template parameters.
+	 * When modifying a runtime-written flag list, do a runtime check on the final flag variables using the function parameters.
+	 *
+	 * @param CompileTimeUnitFlags	The compile-time unit test flag list to be checked
+	 * @param CompileTimeMinFlags	The compile-time minimal client flag list to be checked
+	 * @param RuntimeUnitFlags		The runtime unit test flag list to be checked
+	 * @param RuntimeMinFlags		The runtime minimal client flag list to be checked
+	 */
+	template<EUnitTestFlags CompileTimeUnitFlags=EUnitTestFlags::None, EMinClientFlags CompileTimeMinFlags=EMinClientFlags::None>
+	void ValidateUnitFlags(EUnitTestFlags RuntimeUnitFlags=EUnitTestFlags::None, EMinClientFlags RuntimeMinFlags=EMinClientFlags::None)
+	{
+		// Validate EMinClientFlags
+		ValidateMinFlags<CompileTimeMinFlags>(RuntimeMinFlags);
+
+		PRAGMA_DISABLE_SHADOW_VARIABLE_WARNINGS 
+
+		// Implements the '*Flags' variables within 'Condition', for both static/compile-time checks, and runtime checks
+		#define FLAG_ASSERT(Condition, Message) \
+			{ \
+				struct CompileTimeAssert \
+				{ \
+					void NeverCalled() \
+					{ \
+						constexpr EUnitTestFlags UnitTestFlags = CompileTimeUnitFlags; \
+						constexpr EMinClientFlags MinClientFlags = CompileTimeMinFlags; \
+						static_assert((CompileTimeUnitFlags == EUnitTestFlags::None && CompileTimeMinFlags == EMinClientFlags::None) \
+										|| (Condition), Message); \
+					} \
+				}; \
+				\
+				struct RuntimeAssert \
+				{ \
+					static void DoCheck(EUnitTestFlags UnitTestFlags, EMinClientFlags MinClientFlags) \
+					{ \
+						UNIT_ASSERT(Condition); \
+					} \
+				}; \
+				\
+				if (RuntimeUnitFlags != EUnitTestFlags::None || RuntimeMinFlags != EMinClientFlags::None) \
+				{ \
+					RuntimeAssert::DoCheck(RuntimeUnitFlags, RuntimeMinFlags); \
+				} \
+			}
+
+
+
+		FLAG_ASSERT(!!(UnitTestFlags & EUnitTestFlags::LaunchServer),
+					"Currently, unit tests don't support NOT launching/connecting to a server");
+
+		FLAG_ASSERT((!(UnitTestFlags & EUnitTestFlags::AcceptPlayerController) && !(UnitTestFlags & EUnitTestFlags::RequireNUTActor)) ||
+						!!(MinClientFlags & EMinClientFlags::AcceptActors),
+					"If you require a player/NUTActor, you need to accept actor channels");
+
+		FLAG_ASSERT(!(UnitTestFlags & EUnitTestFlags::RequireNUTActor) || !!(UnitTestFlags & EUnitTestFlags::AcceptPlayerController) ||
+					!!(UnitTestFlags & EUnitTestFlags::RequireBeacon),
+					"If you require a NUTActor, you need to either accept a PlayerController or require a beacon");
+
+		FLAG_ASSERT(!(UnitTestFlags & EUnitTestFlags::RequirePlayerController) ||
+						!!(UnitTestFlags & EUnitTestFlags::AcceptPlayerController),
+					"Don't require a PlayerController, if you don't accept one");
+
+		FLAG_ASSERT(!(UnitTestFlags & EUnitTestFlags::RequirePawn) || !!(UnitTestFlags & EUnitTestFlags::RequirePlayerController),
+					"If you require a pawn, you must require a PlayerController");
+
+		FLAG_ASSERT(!(UnitTestFlags & EUnitTestFlags::RequirePawn) || !!(MinClientFlags & EMinClientFlags::NotifyProcessNetEvent),
+					"If you require a pawn, you must enable NotifyProcessNetEvent");
+
+		FLAG_ASSERT(!(UnitTestFlags & EUnitTestFlags::RequirePlayerState) ||
+					!!(UnitTestFlags & EUnitTestFlags::RequirePlayerController),
+					"If you require a PlayerState, you must require a PlayerController");
+
+		FLAG_ASSERT(!(UnitTestFlags & EUnitTestFlags::RequirePawn) || !!(MinClientFlags & EMinClientFlags::NotifyNetActors),
+					"For part of pawn-setup detection, you need notification for net actors");
+
+		FLAG_ASSERT(!(UnitTestFlags & EUnitTestFlags::RequirePlayerState) || !!(MinClientFlags & EMinClientFlags::NotifyNetActors),
+					"For part of PlayerState-setup detection, you need notification for net actors");
+
+		// As above, but with the 'custom' requirements flag
+		// NOTE: Removed, as it can still be useful to have other requirements flags
+		//FLAG_ASSERT(!(UnitTestFlags & EUnitTestFlags::RequireCustom) ||
+		//				FMath::IsPowerOfTwo((uint32)(UnitTestFlags & EUnitTestFlags::RequirementsMask)));
+
+		FLAG_ASSERT(!(MinClientFlags & EMinClientFlags::SendRPCs) || !!(UnitTestFlags & EUnitTestFlags::AcceptPlayerController) ||
+					!!(UnitTestFlags & EUnitTestFlags::BeaconConnect),
+					"You can't send RPC's, without accepting a player controller (netcode blocks this, without a PC); "
+					"unless this is a beacon");
+
+		// If connecting to a beacon, a number of unit test flags are not supported
+		const EUnitTestFlags RejectedBeaconFlags = (EUnitTestFlags::AcceptPlayerController | EUnitTestFlags::RequirePlayerController |
+													EUnitTestFlags::RequirePing);
+
+		FLAG_ASSERT(!(UnitTestFlags & EUnitTestFlags::BeaconConnect) || !(UnitTestFlags & RejectedBeaconFlags),
+					"Some unit test flags are incompatible with EUnitTestFlags::BeaconConnect");
+
+		FLAG_ASSERT(!(UnitTestFlags & EUnitTestFlags::BeaconConnect) || !!(MinClientFlags & EMinClientFlags::NotifyNetActors),
+					"If connecting to a beacon, net actor notification is required, for proper setup");
+
+		FLAG_ASSERT(!(UnitTestFlags & EUnitTestFlags::RequireBeacon) || !!(UnitTestFlags & EUnitTestFlags::BeaconConnect),
+					"Don't require a beacon, if you're not connecting to a beacon");
+
+		FLAG_ASSERT(!(UnitTestFlags & EUnitTestFlags::LaunchClient) || !!(UnitTestFlags & EUnitTestFlags::LaunchServer),
+					"Don't specify server-dependent flags, if not auto-launching a server");
+
+		FLAG_ASSERT(!(UnitTestFlags & EUnitTestFlags::RequireNUTActor) || !!(MinClientFlags & EMinClientFlags::NotifyNetActors),
+					"You can't use 'RequireNUTActor', without net actor notifications");
+
+		FLAG_ASSERT(!(UnitTestFlags & EUnitTestFlags::ExpectServerCrash) || !!(UnitTestFlags & EUnitTestFlags::ExpectDisconnect),
+					"If a unit test expects a server crash, it should also expect a disconnect too "
+					"(to avoid an invalid 'unit test needs update' result)");
+
+		#undef FLAG_ASSERT
+
+		PRAGMA_ENABLE_SHADOW_VARIABLE_WARNINGS 
+	}
+
+	/**
+	 * Sets and validates at compile time, that the specified flags are valid.
+	 * NOTE: If your unit test subclasses another, it will have to manually change UnitTestFlags/MinClientFlags at runtime.
+	 *
+	 * @param CompileTimeUnitFlags	The compile-time flag list to be checked
+	 * @param CompileTimeMinFlags	The compile-time flag list to be checked
+	 */
+	template<EUnitTestFlags CompileTimeUnitFlags=EUnitTestFlags::None, EMinClientFlags CompileTimeMinFlags=EMinClientFlags::None>
+	void SetFlags()
+	{
+		ValidateUnitFlags<CompileTimeUnitFlags, CompileTimeMinFlags>();
+
+		UnitTestFlags = CompileTimeUnitFlags;
+		MinClientFlags = CompileTimeMinFlags;
+	}
+
 	virtual bool ValidateUnitTestSettings(bool bCDOCheck=false) override;
 
 	virtual void ResetTimeout(FString ResetReason, bool bResetConnTimeout=false, uint32 MinDuration=0) override;
@@ -619,17 +602,17 @@ protected:
 
 
 	/**
-	 * Connects a fake client, to the launched/launching server
+	 * Connects a minimal client, to the launched/launching server
 	 *
 	 * @param InNetID		The unique net id the player should use
 	 * @return				Whether or not the connection kicked off successfully
 	 */
-	virtual bool ConnectFakeClient(FUniqueNetIdRepl* InNetID=nullptr);
+	virtual bool ConnectMinimalClient(const TCHAR* InNetID=nullptr);
 
 	/**
-	 * Cleans up the local/fake client
+	 * Cleans up the minimal client
 	 */
-	virtual void CleanupFakeClient();
+	virtual void CleanupMinimalClient();
 
 	/**
 	 * Triggers an auto-reconnect (disconnect/reconnect) of the fake client
@@ -670,11 +653,6 @@ protected:
 	virtual void PrintUnitTestProcessErrors(TSharedPtr<FUnitTestProcess> InHandle) override;
 
 
-#if !UE_BUILD_SHIPPING
-	static bool InternalScriptProcessEvent(AActor* Actor, UFunction* Function, void* Parameters, void* HookOrigin);
-#endif
-
-
 	virtual void UnitTick(float DeltaTime) override;
 
 	virtual bool IsTickable() const override;
diff --git a/Engine/Plugins/NetcodeUnitTest/NetcodeUnitTest/Source/NetcodeUnitTest/Classes/MinimalClient.h b/Engine/Plugins/NetcodeUnitTest/NetcodeUnitTest/Source/NetcodeUnitTest/Classes/MinimalClient.h
--- a/Engine/Plugins/NetcodeUnitTest/NetcodeUnitTest/Source/NetcodeUnitTest/Classes/MinimalClient.h
+++ b/Engine/Plugins/NetcodeUnitTest/NetcodeUnitTest/Source/NetcodeUnitTest/Classes/MinimalClient.h
@@ -4,70 +4,15 @@
 
 #include "CoreMinimal.h"
 
+#include "Engine/NetConnection.h"
 #include "Engine/PendingNetGame.h"
 
-#include "ClientUnitTest.h"
+#include "Net/NUTUtilNet.h"
+#include "NUTEnum.h"
 
 #include "MinimalClient.generated.h"
 
 
-/**
- * Flags for configuring the minimal client, what parts of the netcode should be enabled etc.
- */
-UENUM()
-enum class EMinClientFlags : uint32
-{
-	None					= 0x00000000,	// No flags
-
-	/** Minimal-client netcode functionality */
-	AcceptActors			= 0x00000004,	// Whether or not to accept actor channels (acts as whitelist-only with NotifyAllowNetActor)
-	AcceptRPCs				= 0x00000010,	// Whether or not to accept execution of any actor RPC's (they are all blocked by default)
-	SendRPCs				= 0x00000020,	// Whether or not to allow RPC sending (NOT blocked by default, @todo JohnB: Add send hook)
-	SkipControlJoin			= 0x00000040,	// Whether or not to skip sending NMT_Join upon connect (or NMT_BeaconJoin for beacons)
-	BeaconConnect			= 0x00000080,	// Whether or not to connect to the servers beacon (greatly limits the connection)
-
-	/** Unit test events */
-	NotifyNetActors			= 0x00400000,	// Whether or not to trigger a 'NotifyNetActor' event, AFTER creation of actor channel actor
-	// @todo #JohnB: You want to avoid duplicate ProcessEvent hooks, due to performance issues
-	NotifyProcessEvent		= 0x00800000,	// Whether or not to trigger 'NotifyScriptProcessEvent' for every executed local function
-
-	/** Debugging */
-	DumpReceivedRaw			= 0x04000000,	// Whether or not to also hex-dump the raw packet receives to the log/log-window
-	DumpSendRaw				= 0x08000000,	// Whether or not to also hex-dump the raw packet sends to the log/log-window
-	DumpSendRPC				= 0x40000000	// Whether or not to dump RPC sends
-};
-
-// Required for bitwise operations with the above enum
-ENUM_CLASS_FLAGS(EMinClientFlags);
-
-// @todo #JohnB: Make sure that all flags contained above, are also marked in the FromUnitTestFlags function (if appropriate)
-
-/**
- * Converts any EUnitTestFlags values, to their EMinClientFlags equivalent
- */
-inline EMinClientFlags FromUnitTestFlags(EUnitTestFlags Flags)
-{
-	EMinClientFlags ReturnFlags = EMinClientFlags::None;
-
-	#define CHECK_FLAG(x) ReturnFlags |= !!(Flags & EUnitTestFlags::x) ? EMinClientFlags::x : EMinClientFlags::None
-
-	CHECK_FLAG(AcceptActors);
-	CHECK_FLAG(AcceptRPCs);
-	CHECK_FLAG(SendRPCs);
-	CHECK_FLAG(SkipControlJoin);
-	CHECK_FLAG(BeaconConnect);
-	CHECK_FLAG(NotifyNetActors);
-	CHECK_FLAG(NotifyProcessEvent);
-	CHECK_FLAG(DumpReceivedRaw);
-	CHECK_FLAG(DumpSendRaw);
-	CHECK_FLAG(DumpSendRPC);
-
-	#undef CHECK_FLAG
-
-	return ReturnFlags;
-}
-
-
 // Delegates
 
 /**
@@ -83,20 +28,6 @@ DECLARE_DELEGATE(FOnMinClientConnected);
 */
 DECLARE_DELEGATE_TwoParams(FOnMinClientNetworkFailure, ENetworkFailure::Type /*FailureType*/, const FString& /*ErrorString*/);
 
-/**
- * Delegate used to control/block sending of RPC's
- *
- * @param bAllowRPC		Whether or not to allow sending of the RPC
- * @param Actor			The actor the RPC will be called in
- * @param Function		The RPC to call
- * @param Parameters	The parameters data blob
- * @param OutParms		Out parameter information (irrelevant for RPC's)
- * @param Stack			The script stack
- * @param SubObject		The sub-object the RPC is being called in (if applicable)
-*/
-DECLARE_DELEGATE_SevenParams(FOnMinClientSendRPC, bool& /*bAllowRPC*/, AActor* /*Actor*/, UFunction* /*Function*/,
-									void* /*Parameters*/, FOutParmRec* /*OutParms*/, FFrame* /*Stack*/, UObject* /*SubObject*/);
-
 /**
  * Delegate for hooking the control channel's 'ReceivedBunch' call
  *
@@ -113,32 +44,14 @@ DECLARE_DELEGATE_OneParam(FOnMinClientReceivedControlBunch, FInBunch&/* Bunch*/)
 */
 DECLARE_DELEGATE_TwoParams(FOnMinClientReceivedRawPacket, void* /*Data*/, int32& /*Count*/);
 
-/**
- * Delegate for hooking the net connections 'LowLevelSend'
- *
- * @param Data			The data being sent
- * @param Count			The number of bytes being sent
- * @param bBlockSend	Whether or not to block the send (defaults to false)
-*/
-DECLARE_DELEGATE_ThreeParams(FOnMinClientLowLevelSend, void* /*Data*/, int32 /*Count*/, bool& /*bBlockSend*/);
-
-/**
- * Delegate for hooking the socket level 'SendTo'
- *
- * @param Data			The data being sent
- * @param Count			The number of bytes being sent
- * @param bBlockSend	Whether or not to block the send (defaults to false)
-*/
-DECLARE_DELEGATE_ThreeParams(FOnMinClientSocketSend, void* /*Data*/, int32 /*Count*/, bool& /*bBlockSend*/);
-
 /**
  * Delegate for notifying on (and optionally blocking) replicated actor creation
  *
  * @param ActorClass	The class of the actor being replicated
  * @param bActorChannel	Whether or not this actor creation is from an actor channel
- * @return				Whether or not to allow creation of the actor
+ * @param bBlockActor	Whether or not to block creation of the actor (defaults to true)
 */
-DECLARE_DELEGATE_RetVal_TwoParams(bool, FOnMinClientRepActorSpawn, UClass* /*ActorClass*/, bool /*bActorChannel*/);
+DECLARE_DELEGATE_ThreeParams(FOnMinClientRepActorSpawn, UClass* /*ActorClass*/, bool /*bActorChannel*/, bool& /*bBlockActor*/);
 
 /**
  * Delegate for notifying AFTER an actor channel actor has been created
@@ -148,6 +61,14 @@ DECLARE_DELEGATE_RetVal_TwoParams(bool, FOnMinClientRepActorSpawn, UClass* /*Act
 */
 DECLARE_DELEGATE_TwoParams(FOnMinClientNetActor, UActorChannel* /*ActorChannel*/, AActor* /*Actor*/);
 
+/**
+ * Delegate for hooking the HandlerClientPlayer event
+ *
+ * @param PC			The PlayerController being initialized with the net connection
+ * @param Connection	The net connection the player is being initialized with
+ */
+DECLARE_DELEGATE_TwoParams(FOnHandleClientPlayer, APlayerController* /*PC*/, UNetConnection* /*Connection*/);
+
 
 /**
  * Parameters for configuring the minimal client - also directly inherited by UMinimalClient
@@ -174,7 +95,10 @@ struct FMinClientParms
 	FString BeaconType;
 
 	/** If overriding the UID used for joining, this specifies it */
-	FUniqueNetIdRepl* JoinUID;
+	FString JoinUID;
+
+	/** Clientside RPC's that should be allowed to execute (requires the NotifyProcessNetEvent flag) */
+	TArray<FString> AllowedClientRPCs;
 
 
 	/**
@@ -186,11 +110,16 @@ struct FMinClientParms
 		, ServerAddress()
 		, BeaconAddress()
 		, BeaconType()
-		, JoinUID(nullptr)
+		, JoinUID(TEXT("Dud"))
 	{
 	}
 
 protected:
+	/**
+	 * Verify that the parameters specified to this struct are valid
+	 */
+	void ValidateParms();
+
 	void CopyParms(FMinClientParms& Target)
 	{
 		Target.MinClientFlags = MinClientFlags;
@@ -199,6 +128,7 @@ protected:
 		Target.BeaconAddress = BeaconAddress;
 		Target.BeaconType = BeaconType;
 		Target.JoinUID = JoinUID;
+		Target.AllowedClientRPCs = AllowedClientRPCs;
 	}
 };
 
@@ -217,8 +147,13 @@ struct FMinClientHooks
 	/** Delegate notifying of network failure */
 	FOnMinClientNetworkFailure NetworkFailureDel;
 
+	/** Delegate for notifying/controlling RPC receives */
+	FOnProcessNetEvent ReceiveRPCDel;
+
+#if !UE_BUILD_SHIPPING
 	/** Delegate for notifying/controlling RPC sends */
-	FOnMinClientSendRPC SendRPCDel;
+	FOnSendRPC SendRPCDel;
+#endif
 
 	/** Delegate for notifying of control channel bunches */
 	FOnMinClientReceivedControlBunch ReceivedControlBunchDel;
@@ -226,11 +161,10 @@ struct FMinClientHooks
 	/** Delegate for notifying of net connection raw packet receives */
 	FOnMinClientReceivedRawPacket ReceivedRawPacketDel;
 
+#if !UE_BUILD_SHIPPING
 	/** Delegate for notifying of net connection low level packet sends */
-	FOnMinClientLowLevelSend LowLevelSendDel;
-
-	/** Delegate for notifying of socket level packet sends */
-	FOnMinClientSocketSend SocketSendDel;
+	FOnLowLevelSend LowLevelSendDel;
+#endif
 
 	/** Delegate for notifying/controlling replicated actor spawning */
 	FOnMinClientRepActorSpawn RepActorSpawnDel;
@@ -238,6 +172,9 @@ struct FMinClientHooks
 	/** Delegate for notifying AFTER net actor creation */
 	FOnMinClientNetActor NetActorDel;
 
+	/** Delegate for notifying of the net connection HandlerClientPlayer event */
+	FOnHandleClientPlayer HandleClientPlayerDel;
+
 
 	/**
 	 * Default constructor
@@ -245,13 +182,18 @@ struct FMinClientHooks
 	FMinClientHooks()
 		: ConnectedDel()
 		, NetworkFailureDel()
+		, ReceiveRPCDel()
+#if !UE_BUILD_SHIPPING
 		, SendRPCDel()
+#endif
 		, ReceivedControlBunchDel()
 		, ReceivedRawPacketDel()
+#if !UE_BUILD_SHIPPING
 		, LowLevelSendDel()
-		, SocketSendDel()
+#endif
 		, RepActorSpawnDel()
 		, NetActorDel()
+		, HandleClientPlayerDel()
 	{
 	}
 
@@ -260,13 +202,18 @@ protected:
 	{
 		Target.ConnectedDel = ConnectedDel;
 		Target.NetworkFailureDel = NetworkFailureDel;
+		Target.ReceiveRPCDel = ReceiveRPCDel;
+#if !UE_BUILD_SHIPPING
 		Target.SendRPCDel = SendRPCDel;
+#endif
 		Target.ReceivedControlBunchDel = ReceivedControlBunchDel;
 		Target.ReceivedRawPacketDel = ReceivedRawPacketDel;
+#if !UE_BUILD_SHIPPING
 		Target.LowLevelSendDel = LowLevelSendDel;
-		Target.SocketSendDel = SocketSendDel;
+#endif
 		Target.RepActorSpawnDel = RepActorSpawnDel;
 		Target.NetActorDel = NetActorDel;
+		Target.HandleClientPlayerDel = HandleClientPlayerDel;
 	}
 };
 
@@ -284,8 +231,8 @@ class NETCODEUNITTEST_API UMinimalClient : public UObject, public FNetworkNotify
 {
 	friend class UUnitTestNetConnection;
 	friend class UUnitTestChannel;
+	friend class UUnitTestActorChannel;
 	friend class UUnitTestPackageMap;
-	friend class FSocketHook;
 
 	GENERATED_UCLASS_BODY()
 
@@ -298,12 +245,31 @@ public:
 	 * @param Hooks		The delegates for hooking low level netcode events in the minimal client.
 	 * @return			Whether or not the connection kicked off successfully.
 	 */
-	virtual bool Connect(FMinClientParms Parms, FMinClientHooks Hooks);
+	bool Connect(FMinClientParms Parms, FMinClientHooks Hooks);
 
 	/**
 	 * Disconnects and cleans up the minimal client.
 	 */
-	virtual void Cleanup();
+	void Cleanup();
+
+
+	/**
+	 * Creates a bunch for the specified channel, with the ability to create the channel as well
+	 * WARNING: Can return nullptr! (e.g. if the control channel is saturated)
+	 *
+	 * @param ChType			Specify the type of channel the bunch will be sent on
+	 * @param ChIndex			Specify the index of the channel to send the bunch on (otherwise, picks the next free/unused channel)
+	 * @return					Returns the created bunch, or nullptr if it couldn't be created
+	 */
+	FOutBunch* CreateChannelBunch(EChannelType ChType, int32 ChIndex=INDEX_NONE);
+
+	/**
+	 * Sends a bunch over the control channel
+	 *
+	 * @param ControlChanBunch	The bunch to send over the control channel
+	 * @return					Whether or not the bunch was sent successfully
+	 */
+	bool SendControlBunch(FOutBunch* ControlChanBunch);
 
 
 	// Accessors
@@ -340,24 +306,30 @@ public:
 		return UnitConn;
 	}
 
+	/**
+	 * Whether or not the minimal client is connected to the server
+	 */
+	FORCEINLINE bool IsConnected()
+	{
+		return bConnected;
+	}
+
 
 	/**
-	 * Ticks the minimal client, through the owner unit tests UnitTick function
+	 * Matches up an active NetConnection to a MinimalClient, if run through a unit test
 	 */
-	virtual void UnitTick(float DeltaTime);
+	static UMinimalClient* GetMinClientFromConn(UNetConnection* InConn);
+
 
 	/**
-	 * Whether or not the minimal client requires ticking.
+	 * Ticks the minimal client, through the owner unit tests UnitTick function
 	 */
-	bool IsTickable() const;
+	void UnitTick(float DeltaTime);
 
 	/**
-	 * Whether or not the minimal client is connected to the server
+	 * Whether or not the minimal client requires ticking.
 	 */
-	FORCEINLINE bool IsConnected()
-	{
-		return bConnected;
-	}
+	bool IsTickable() const;
 
 
 protected:
@@ -368,6 +340,11 @@ protected:
 	 */
 	bool ConnectMinimalClient();
 
+	/**
+	 * Creates a net driver for the minimal client
+	 */
+	void CreateNetDriver();
+
 	/**
 	 * Sends the packet for triggering the initial join (usually is delayed, by the PacketHandler)
 	 */
@@ -380,10 +357,26 @@ protected:
 	void DisconnectMinimalClient();
 
 	/**
-	 * See FOnMinClientSendRPC
+	 * See FOnLowLevelSend
 	 */
-	virtual bool NotifySendRPC(AActor* Actor, UFunction* Function, void* Parameters, FOutParmRec* OutParms, FFrame* Stack,
-								UObject* SubObject);
+	void NotifySocketSend(void* Data, int32 Count, bool& bBlockSend);
+
+	/**
+	 * See FOnReceivedRawPacket
+	 */
+	void NotifyReceivedRawPacket(void* Data, int32 Count, bool& bBlockReceive);
+
+	/**
+	 * See FOnProcessNetEvent
+	 */
+	void NotifyReceiveRPC(AActor* Actor, UFunction* Function, void* Parameters, bool& bBlockRPC);
+
+
+	/**
+	 * See FOnSendRPC
+	 */
+	void NotifySendRPC(AActor* Actor, UFunction* Function, void* Parameters, FOutParmRec* OutParms, FFrame* Stack, UObject* SubObject,
+						bool& bBlockSendRPC);
 
 
 protected:
@@ -418,15 +411,12 @@ protected:
 	/** Stores a reference to the created fake world, for execution and later cleanup */
 	UWorld* UnitWorld;
 
-	/** Stores a reference to the created unit test net driver, for execution and later cleanup (always a UUnitTestNetDriver) */
+	/** Stores a reference to the created unit test net driver, for execution and later cleanup */
 	UNetDriver* UnitNetDriver;
 
 	/** Stores a reference to the server connection (always a 'UUnitTestNetConnection') */
 	UNetConnection* UnitConn;
 
-	/** For the unit test control channel, this tracks the current bunch sequence */
-	int32 ControlBunchSequence;
-
 	/** If notifying of net actor creation, this keeps track of new actor channel indexes pending notification */
 	TArray<int32> PendingNetActorChans;
 
diff --git a/Engine/Plugins/NetcodeUnitTest/NetcodeUnitTest/Source/NetcodeUnitTest/Classes/NUTActor.h b/Engine/Plugins/NetcodeUnitTest/NetcodeUnitTest/Source/NetcodeUnitTest/Classes/NUTActor.h
--- a/Engine/Plugins/NetcodeUnitTest/NetcodeUnitTest/Source/NetcodeUnitTest/Classes/NUTActor.h
+++ b/Engine/Plugins/NetcodeUnitTest/NetcodeUnitTest/Source/NetcodeUnitTest/Classes/NUTActor.h
@@ -18,25 +18,22 @@ class ANUTActor;
  * Enum for defining custom NetcodeUnitTest control channel commands (sent through NMT_NUTControl)
  */
 // @todo #JohnBDoc: Fully document individually
-namespace ENUTControlCommand
+enum class ENUTControlCommand : uint8
 {
-	enum Type
-	{
-		Command_NoResult		= 0,
-		Command_SendResult		= 1,
-		CommandResult_Failed	= 2,
-		CommandResult_Success	= 3,
-		Ping					= 4,
-		Pong					= 5,
-		WatchEvent				= 6,
-		NotifyEvent				= 7,
-		Summon					= 8,
-		SuspendProcess			= 9
-	};
-}
+	Command_NoResult		= 0,
+	Command_SendResult		= 1,
+	CommandResult_Failed	= 2,
+	CommandResult_Success	= 3,
+	Ping					= 4,
+	Pong					= 5,
+	WatchEvent				= 6,
+	NotifyEvent				= 7,
+	Summon					= 8,
+	SuspendProcess			= 9
+};
 
 // Custom control channel message, used to communicate with the server NUTActor over the control channel
-DEFINE_CONTROL_CHANNEL_MESSAGE_TWOPARAM(NUTControl, 250, uint8, FString);
+DEFINE_CONTROL_CHANNEL_MESSAGE_TWOPARAM(NUTControl, 250, ENUTControlCommand, FString);
 
 
 /**
diff --git a/Engine/Plugins/NetcodeUnitTest/NetcodeUnitTest/Source/NetcodeUnitTest/Classes/Net/UnitTestActorChannel.h b/Engine/Plugins/NetcodeUnitTest/NetcodeUnitTest/Source/NetcodeUnitTest/Classes/Net/UnitTestActorChannel.h
--- a/Engine/Plugins/NetcodeUnitTest/NetcodeUnitTest/Source/NetcodeUnitTest/Classes/Net/UnitTestActorChannel.h
+++ b/Engine/Plugins/NetcodeUnitTest/NetcodeUnitTest/Source/NetcodeUnitTest/Classes/Net/UnitTestActorChannel.h
@@ -8,7 +8,11 @@
 
 #include "UnitTestActorChannel.generated.h"
 
+
+// Forward declarations
 class FInBunch;
+class UMinimalClient;
+
 
 /**
  * An actor net channel override, for hooking ReceivedBunch, to aid in detecting/blocking of remote actors, of a specific class
@@ -18,8 +22,16 @@ class UUnitTestActorChannel : public UActorChannel
 {
 	GENERATED_UCLASS_BODY()
 
+	virtual void Init(UNetConnection* InConnection, int32 InChIndex, bool InOpenedLocally) override;
+
 	virtual void ReceivedBunch(FInBunch& Bunch) override;
 
 	virtual void Tick() override;
+
+	virtual void NotifyActorChannelOpen(AActor* InActor, FInBunch& InBunch) override;
+
+private:
+	/** Cached referenced to the minimal client that owns this actor channel */
+	UMinimalClient* MinClient;
 };
 
diff --git a/Engine/Plugins/NetcodeUnitTest/NetcodeUnitTest/Source/NetcodeUnitTest/Classes/Net/UnitTestPackageMap.h b/Engine/Plugins/NetcodeUnitTest/NetcodeUnitTest/Source/NetcodeUnitTest/Classes/Net/UnitTestPackageMap.h
--- a/Engine/Plugins/NetcodeUnitTest/NetcodeUnitTest/Source/NetcodeUnitTest/Classes/Net/UnitTestPackageMap.h
+++ b/Engine/Plugins/NetcodeUnitTest/NetcodeUnitTest/Source/NetcodeUnitTest/Classes/Net/UnitTestPackageMap.h
@@ -8,7 +8,11 @@
 #include "NetcodeUnitTest.h"
 #include "UnitTestPackageMap.generated.h"
 
+
+// Forward declarations
 class UNetConnection;
+class UMinimalClient;
+
 
 /**
  * Package map override, for blocking the creation of actor channels for specific actors (by detecting the actor class being created)
@@ -34,6 +38,9 @@ class UUnitTestPackageMap : public UPackageMapClient
 	virtual bool SerializeNewActor(FArchive& Ar, class UActorChannel* Channel, class AActor*& Actor) override;
 
 public:
+	/** Cached reference to the minimal client that owns this package map */
+	UMinimalClient* MinClient;
+
 	/** Whether or not we are currently within execution of SerializeNewActor */
 	bool bWithinSerializeNewActor;
 
diff --git a/Engine/Plugins/NetcodeUnitTest/NetcodeUnitTest/Source/NetcodeUnitTest/NetcodeUnitTest.Build.cs b/Engine/Plugins/NetcodeUnitTest/NetcodeUnitTest/Source/NetcodeUnitTest/NetcodeUnitTest.Build.cs
--- a/Engine/Plugins/NetcodeUnitTest/NetcodeUnitTest/Source/NetcodeUnitTest/NetcodeUnitTest.Build.cs
+++ b/Engine/Plugins/NetcodeUnitTest/NetcodeUnitTest/Source/NetcodeUnitTest/NetcodeUnitTest.Build.cs
@@ -15,7 +15,6 @@ namespace UnrealBuildTool.Rules
 					"Core",
 					"CoreUObject",
 					"Engine",
-					"OnlineSubsystemUtils",
 					"Sockets"
 				}
 			);
