Copying //UE4/Dev-Networking to //UE4/Dev-Main (Source: //UE4/Dev-Networking @ 3596412)

diff --git a/Engine/Plugins/NetcodeUnitTest/NetcodeUnitTest/Source/NetcodeUnitTest/Private/ClientUnitTest.cpp b/Engine/Plugins/NetcodeUnitTest/NetcodeUnitTest/Source/NetcodeUnitTest/Private/ClientUnitTest.cpp
--- a/Engine/Plugins/NetcodeUnitTest/NetcodeUnitTest/Source/NetcodeUnitTest/Private/ClientUnitTest.cpp
+++ b/Engine/Plugins/NetcodeUnitTest/NetcodeUnitTest/Source/NetcodeUnitTest/Private/ClientUnitTest.cpp
@@ -8,33 +8,25 @@
 #include "GameFramework/PlayerState.h"
 #include "GameFramework/PlayerController.h"
 #include "HAL/PlatformNamedPipe.h"
-
-#include "OnlineBeaconClient.h"
 #include "Engine/ActorChannel.h"
 
 
 #include "UnitTestManager.h"
 #include "MinimalClient.h"
 #include "NUTActor.h"
-
-#include "Net/UnitTestNetConnection.h"
 #include "Net/UnitTestChannel.h"
-
 #include "UnitTestEnvironment.h"
 #include "NUTUtilDebug.h"
 #include "NUTUtilReflection.h"
 
 
+UClass* UClientUnitTest::OnlineBeaconClass = FindObject<UClass>(ANY_PACKAGE, TEXT("OnlineBeaconClient"));
 
 // @todo #JohnBMultiFakeClient: Eventually, move >all< of the minimal/headless client handling code, into a new/separate class,
 //				so that a single unit test can have multiple minimal clients on a server.
 //				This would be useful for licensees, for doing load testing:
 //				https://udn.unrealengine.com/questions/247014/clientserver-automation.html
 
-
-// @todo #JohnB: Consider a new class of 'notify' flags, for enabling/disabling the notify functions,
-//					and streamlining/optimizing the unit test flags
-
 /**
  * UClientUnitTest
  */
@@ -42,6 +34,7 @@
 UClientUnitTest::UClientUnitTest(const FObjectInitializer& ObjectInitializer)
 	: Super(ObjectInitializer)
 	, UnitTestFlags(EUnitTestFlags::None)
+	, MinClientFlags(EMinClientFlags::None)
 	, BaseServerURL(TEXT(""))
 	, BaseServerParameters(TEXT(""))
 	, BaseClientURL(TEXT(""))
@@ -57,11 +50,6 @@ UClientUnitTest::UClientUnitTest(const FObjectInitializer& ObjectInitializer)
 	, bBlockingFakeClientDelay(false)
 	, NextBlockingTimeout(0.0)
 	, MinClient(nullptr)
-	// @todo #JohnB: Deprecate this, after migrating to the minimal client code
-#if 1
-	, UnitNetDriver(nullptr)
-	, UnitConn(nullptr)
-#endif
 	, bTriggerredInitialConnect(false)
 	, UnitPC(nullptr)
 	, bUnitPawnSetup(false)
@@ -70,12 +58,9 @@ UClientUnitTest::UClientUnitTest(const FObjectInitializer& ObjectInitializer)
 	, bUnitNUTActorSetup(false)
 	, UnitBeacon(nullptr)
 	, bReceivedPong(false)
-	// @todo #JohnB: Remove after deprecated
-	//, ControlBunchSequence(0)
-	// @todo #JohnB: Remove once fully migrated to minimal client
-	, ControlBunchSequence(nullptr)
 	, bPendingNetworkFailure(false)
 	, bDetectedMCPOnline(false)
+	, bSentBunch(false)
 {
 }
 
@@ -84,8 +69,14 @@ void UClientUnitTest::NotifyMinClientConnected()
 {
 	if (HasAllRequirements())
 	{
+		ResetTimeout(TEXT("ExecuteClientUnitTest (NotifyMinClientConnected)"));
 		ExecuteClientUnitTest();
 	}
+
+	if (!!(UnitTestFlags & EUnitTestFlags::RequirePing))
+	{
+		SendNUTControl(ENUTControlCommand::Ping, TEXT(""));
+	}
 }
 
 void UClientUnitTest::NotifyControlMessage(FInBunch& Bunch, uint8 MessageType)
@@ -120,40 +111,40 @@ void UClientUnitTest::NotifyHandleClientPlayer(APlayerController* PC, UNetConnec
 	}
 }
 
-bool UClientUnitTest::NotifyAllowNetActor(UClass* ActorClass, bool bActorChannel)
+void UClientUnitTest::NotifyAllowNetActor(UClass* ActorClass, bool bActorChannel, bool& bBlockActor)
 {
-	bool bAllow = false;
-
 	if (!!(UnitTestFlags & EUnitTestFlags::RequireNUTActor) && ActorClass == ANUTActor::StaticClass() && UnitNUTActor == nullptr)
 	{
-		bAllow = true;
+		bBlockActor = false;
 	}
 
 	if (!!(UnitTestFlags & EUnitTestFlags::AcceptPlayerController) && ActorClass->IsChildOf(APlayerController::StaticClass()) &&
 		UnitPC == nullptr)
 	{
-		bAllow = true;
+		bBlockActor = false;
 	}
 
 	if (!!(UnitTestFlags & EUnitTestFlags::RequirePawn) && ActorClass->IsChildOf(ACharacter::StaticClass()) &&
 		(!UnitPC.IsValid() || UnitPC->GetCharacter() == nullptr))
 	{
-		bAllow = true;
+		bBlockActor = false;
 	}
 
 	if (!!(UnitTestFlags & EUnitTestFlags::RequirePlayerState) && ActorClass->IsChildOf(APlayerState::StaticClass()) &&
 		(!UnitPC.IsValid() || UnitPC->PlayerState == nullptr))
 	{
-		bAllow = true;
+		bBlockActor = false;
 	}
 
-	if (!!(UnitTestFlags & EUnitTestFlags::RequireBeacon) && ActorClass->IsChildOf(AOnlineBeaconClient::StaticClass()) &&
-		UnitBeacon == nullptr)
+	check(OnlineBeaconClass != nullptr);
+
+	if (!!(UnitTestFlags & EUnitTestFlags::RequireBeacon) && ActorClass->IsChildOf(OnlineBeaconClass) && UnitBeacon == nullptr)
 	{
-		bAllow = true;
+		bBlockActor = false;
 	}
 
-	if (!bAllow && AllowedClientActors.Num() > 0)
+	// @todo #JohnB: Move to minimal client, similar to how AllowClientRPCs was moved? (not needed yet...)
+	if (bBlockActor && AllowedClientActors.Num() > 0)
 	{
 		const auto CheckIsChildOf =
 			[&](const UClass* CurEntry)
@@ -162,10 +153,8 @@ bool UClientUnitTest::NotifyAllowNetActor(UClass* ActorClass, bool bActorChannel
 			};
 
 		// Use 'ContainsByPredicate' as iterator
-		bAllow = AllowedClientActors.ContainsByPredicate(CheckIsChildOf);
+		bBlockActor = !AllowedClientActors.ContainsByPredicate(CheckIsChildOf);
 	}
-
-	return bAllow;
 }
 
 void UClientUnitTest::NotifyNetActor(UActorChannel* ActorChannel, AActor* Actor)
@@ -177,23 +166,23 @@ void UClientUnitTest::NotifyNetActor(UActorChannel* ActorChannel, AActor* Actor)
 
 		if (UnitNUTActor.IsValid())
 		{
+			// NOTE: ExecuteClientUnitTest triggered for this, in UnitTick - not here.
 			ResetTimeout(TEXT("NotifyNetActor - UnitNUTActor"));
 		}
 	}
 
-	if (!!(UnitTestFlags & EUnitTestFlags::BeaconConnect) && !UnitBeacon.IsValid())
+	check(OnlineBeaconClass != nullptr);
+
+	if (!!(UnitTestFlags & EUnitTestFlags::BeaconConnect) && !UnitBeacon.IsValid() && Actor->IsA(OnlineBeaconClass))
 	{
-		UnitBeacon = Cast<AOnlineBeaconClient>(Actor);
+		UnitBeacon = Actor;
 
-		if (UnitBeacon.IsValid())
-		{
-			NUTNet::HandleBeaconReplicate(UnitBeacon.Get(), UnitConn);
+		NUTNet::HandleBeaconReplicate(UnitBeacon.Get(), MinClient->GetConn());
 
-			if (!!(UnitTestFlags & EUnitTestFlags::RequireBeacon) && HasAllRequirements())
-			{
-				ResetTimeout(TEXT("ExecuteClientUnitTest (NotifyNetActor - UnitBeacon)"));
-				ExecuteClientUnitTest();
-			}
+		if (!!(UnitTestFlags & EUnitTestFlags::RequireBeacon) && HasAllRequirements())
+		{
+			ResetTimeout(TEXT("ExecuteClientUnitTest (NotifyNetActor - UnitBeacon)"));
+			ExecuteClientUnitTest();
 		}
 	}
 
@@ -275,11 +264,16 @@ void UClientUnitTest::NotifyNetworkFailure(ENetworkFailure::Type FailureType, co
 		// Shut down the fake client now (relevant for developer mode)
 		if (VerificationState != EUnitTestVerification::Unverified)
 		{
-			CleanupFakeClient();
+			CleanupMinimalClient();
 		}
 	}
 }
 
+void UClientUnitTest::NotifySocketSendRawPacket(void* Data, int32 Count, bool& bBlockSend)
+{
+	bSentBunch = !bBlockSend;
+}
+
 void UClientUnitTest::ReceivedControlBunch(FInBunch& Bunch)
 {
 	if (!Bunch.AtEnd())
@@ -291,7 +285,7 @@ void UClientUnitTest::ReceivedControlBunch(FInBunch& Bunch)
 		{
 			if (MessageType == NMT_NUTControl)
 			{
-				uint8 CmdType = 0;
+				ENUTControlCommand CmdType;
 				FString Command;
 				FNetControlMessage<NMT_NUTControl>::Receive(Bunch, CmdType, Command);
 
@@ -320,108 +314,63 @@ void UClientUnitTest::ReceivedControlBunch(FInBunch& Bunch)
 	}
 }
 
-#if !UE_BUILD_SHIPPING
-bool UClientUnitTest::NotifyScriptProcessEvent(AActor* Actor, UFunction* Function, void* Parameters)
+void UClientUnitTest::NotifyReceiveRPC(AActor* Actor, UFunction* Function, void* Parameters, bool& bBlockRPC)
 {
-	bool bBlockEvent = false;
-	bool bNetClientRPC = !!(Function->FunctionFlags & FUNC_Net) && !!(Function->FunctionFlags & FUNC_NetClient);
+	FString FuncName = Function->GetName();
 
-	// Handle UnitTestFlags that require RPC monitoring
-	if (bNetClientRPC)
+	// Handle detection and proper setup of the PlayerController's pawn
+	if (!!(UnitTestFlags & EUnitTestFlags::RequirePawn) && !bUnitPawnSetup && UnitPC != nullptr)
 	{
-		FString FuncName = Function->GetName();
-
-		// Whether or not to force acceptance of this RPC
-		bool bForceAccept = false;
-
-		// Handle detection and proper setup of the PlayerController's pawn
-		if (!!(UnitTestFlags & EUnitTestFlags::RequirePawn) && !bUnitPawnSetup && UnitPC != NULL)
+		if (FuncName == TEXT("ClientRestart"))
 		{
-			if (FuncName == TEXT("ClientRestart"))
-			{
-				UNIT_LOG(ELogType::StatusImportant, TEXT("Got ClientRestart"));
+			UNIT_LOG(ELogType::StatusImportant, TEXT("Got ClientRestart"));
 
-				// Trigger the event directly here, and block execution in the original code, so  we can execute code post-ProcessEvent
-				Actor->UObject::ProcessEvent(Function, Parameters);
+			// Trigger the event directly here, and block execution in the original code, so we can execute code post-ProcessEvent
+			Actor->UObject::ProcessEvent(Function, Parameters);
 
 
-				// If the pawn is set, now execute the exploit
-				if (UnitPC->GetCharacter())
-				{
-					bUnitPawnSetup = true;
+			// If the pawn is set, now execute the exploit
+			if (UnitPC->GetCharacter())
+			{
+				bUnitPawnSetup = true;
 
-					ResetTimeout(TEXT("bUnitPawnSetup"));
+				ResetTimeout(TEXT("bUnitPawnSetup"));
 
-					if (HasAllRequirements())
-					{
-						ResetTimeout(TEXT("ExecuteClientUnitTest (bUnitPawnSetup)"));
-						ExecuteClientUnitTest();
-					}
-				}
-				// If the pawn was not set, get the server to check again
-				else
+				if (HasAllRequirements())
 				{
-					FString LogMsg = TEXT("Pawn was not set, sending ServerCheckClientPossession request");
-
-					ResetTimeout(LogMsg);
-					UNIT_LOG(ELogType::StatusImportant, TEXT("%s"), *LogMsg);
-
-					UnitPC->ServerCheckClientPossession();
+					ResetTimeout(TEXT("ExecuteClientUnitTest (bUnitPawnSetup)"));
+					ExecuteClientUnitTest();
 				}
-
-
-				bBlockEvent = true;
 			}
-			// Retries setting the pawn, which will trigger ClientRestart locally, and enters into the above code with the Pawn set
-			else if (FuncName == TEXT("ClientRetryClientRestart"))
+			// If the pawn was not set, get the server to check again
+			else
 			{
-				bBlockEvent = false;
-				bForceAccept = true;
-			}
-		}
-
-
-		bForceAccept = bForceAccept || AllowedClientRPCs.Contains(FuncName);
-
-		// Block RPC's, if they are not accepted
-		if (!(UnitTestFlags & EUnitTestFlags::AcceptRPCs) && !bForceAccept)
-		{
-			FString FuncParms = NUTUtilRefl::FunctionParmsToString(Function, Parameters);
+				FString LogMsg = TEXT("Pawn was not set, sending ServerCheckClientPossession request");
 
-			UNIT_LOG(, TEXT("Blocking receive RPC '%s' for actor '%s'"), *FuncName, *Actor->GetFullName());
+				ResetTimeout(LogMsg);
+				UNIT_LOG(ELogType::StatusImportant, TEXT("%s"), *LogMsg);
 
-			if (FuncParms.Len() > 0)
-			{
-				UNIT_LOG(, TEXT("     '%s' parameters: %s"), *FuncName, *FuncParms);
+				UnitPC->ServerCheckClientPossession();
 			}
 
-			bBlockEvent = true;
-		}
 
-		if (!!(UnitTestFlags & EUnitTestFlags::DumpReceivedRPC) && !bBlockEvent)
+			bBlockRPC = true;
+		}
+		// Retries setting the pawn, which will trigger ClientRestart locally, and enters into the above code with the Pawn set
+		else if (FuncName == TEXT("ClientRetryClientRestart"))
 		{
-			FString FuncParms = NUTUtilRefl::FunctionParmsToString(Function, Parameters);
-
-			UNIT_LOG(ELogType::StatusDebug, TEXT("Received RPC '%s' for actor '%s'"), *FuncName, *Actor->GetFullName());
-
-			if (FuncParms.Len() > 0)
-			{
-				UNIT_LOG(, TEXT("     '%s' parameters: %s"), *FuncName, *FuncParms);
-			}
+			bBlockRPC = false;
 		}
 	}
-
-	return bBlockEvent;
 }
-#endif
 
 void UClientUnitTest::NotifyProcessLog(TWeakPtr<FUnitTestProcess> InProcess, const TArray<FString>& InLogLines)
 {
 	// Get partial log messages that indicate startup progress/completion
-	const TArray<FString>* ServerStartProgressLogs = NULL;
-	const TArray<FString>* ServerReadyLogs = NULL;
-	const TArray<FString>* ServerTimeoutResetLogs = NULL;
-	const TArray<FString>* ClientTimeoutResetLogs = NULL;
+	const TArray<FString>* ServerStartProgressLogs = nullptr;
+	const TArray<FString>* ServerReadyLogs = nullptr;
+	const TArray<FString>* ServerTimeoutResetLogs = nullptr;
+	const TArray<FString>* ClientTimeoutResetLogs = nullptr;
 
 	UnitEnv->GetServerProgressLogs(ServerStartProgressLogs, ServerReadyLogs, ServerTimeoutResetLogs);
 	UnitEnv->GetClientProgressLogs(ClientTimeoutResetLogs);
@@ -455,12 +404,14 @@ void UClientUnitTest::NotifyProcessLog(TWeakPtr<FUnitTestProcess> InProcess, con
 		// each time there is a server log event, that indicates progress in starting up
 		if (!!(UnitTestFlags & EUnitTestFlags::LaunchServer))
 		{
-			if (!bTriggerredInitialConnect && (UnitConn == NULL || UnitConn->State == EConnectionState::USOCK_Pending))
+			UNetConnection* UnitConn = (MinClient != nullptr ? MinClient->GetConn() : nullptr);
+
+			if (!bTriggerredInitialConnect && (UnitConn == nullptr || UnitConn->State == EConnectionState::USOCK_Pending))
 			{
 				if (ServerReadyLogs->ContainsByPredicate(SearchInLogLine))
 				{
 					// Fire off fake client connection
-					if (UnitConn == NULL)
+					if (UnitConn == nullptr)
 					{
 						bool bBlockingProcess = IsBlockingProcessPresent(true);
 
@@ -480,7 +431,7 @@ void UClientUnitTest::NotifyProcessLog(TWeakPtr<FUnitTestProcess> InProcess, con
 							UNIT_LOG(ELogType::StatusImportant, TEXT("%s"), *LogMsg);
 							UNIT_STATUS_LOG(ELogType::StatusVerbose, TEXT("%s"), *LogMsg);
 
-							ConnectFakeClient();
+							ConnectMinimalClient();
 						}
 					}
 
@@ -564,7 +515,7 @@ void UClientUnitTest::NotifyProcessFinished(TWeakPtr<FUnitTestProcess> InProcess
 
 
 				// Immediately cleanup the fake client (don't wait for end-of-life cleanup in CleanupUnitTest)
-				CleanupFakeClient();
+				CleanupMinimalClient();
 
 
 				// If a server exit was unexpected, mark the unit test as broken
@@ -598,55 +549,23 @@ void UClientUnitTest::NotifyProcessFinished(TWeakPtr<FUnitTestProcess> InProcess
 
 void UClientUnitTest::NotifySuspendRequest()
 {
-	// @todo #JohnBFeature: Currently on suspends the server, suspend the client as well, and also add more granularity
-	//						(deciding which to suspend)
-
 #if PLATFORM_WINDOWS
-	TSharedPtr<FUnitTestProcess> CurProcess = (ServerHandle.IsValid() ? ServerHandle.Pin() : NULL);
+	TSharedPtr<FUnitTestProcess> CurProcess = (ServerHandle.IsValid() ? ServerHandle.Pin() : nullptr);
 
 	if (CurProcess.IsValid())
 	{
 		// Suspend request
 		if (CurProcess->SuspendState == ESuspendState::Active)
 		{
-			if (UnitConn != NULL)
+			if (SendNUTControl(ENUTControlCommand::SuspendProcess, TEXT("")))
 			{
-				UUnitTestChannel* UnitControlChan = Cast<UUnitTestChannel>(UnitConn->Channels[0]);
+				NotifyProcessSuspendState(ServerHandle, ESuspendState::Suspended);
 
-				if (UnitControlChan != NULL)
-				{
-					FOutBunch* ControlChanBunch = NUTNet::CreateChannelBunch(ControlBunchSequence, UnitConn, CHTYPE_Control, 0);
-
-					if (ControlChanBunch != nullptr)
-					{
-						uint8 ControlMsg = NMT_NUTControl;
-						uint8 PingCmd = ENUTControlCommand::SuspendProcess;
-						FString Dud = TEXT("");
-
-						*ControlChanBunch << ControlMsg;
-						*ControlChanBunch << PingCmd;
-						*ControlChanBunch << Dud;
-
-						NUTNet::SendControlBunch(UnitConn, *ControlChanBunch);
-
-
-						NotifyProcessSuspendState(ServerHandle, ESuspendState::Suspended);
-
-						UNIT_LOG(, TEXT("Sent suspend request to server (may take time to execute, if server is still starting)."));
-					}
-					else
-					{
-						UNIT_LOG(, TEXT("Failed to create control channel bunch - connection saturated?"));
-					}
-				}
-				else
-				{
-					UNIT_LOG(, TEXT("No control channel, can't suspend (server still starting?)."))
-				}
+				UNIT_LOG(, TEXT("Sent suspend request to server (may take time to execute, if server is still starting)."));
 			}
 			else
 			{
-				UNIT_LOG(, TEXT("Not connected to server, can't suspend."));
+				UNIT_LOG(, TEXT("Failed to send suspend request to server"));
 			}
 		}
 		// Resume request
@@ -693,8 +612,6 @@ void UClientUnitTest::NotifyProcessSuspendState(TWeakPtr<FUnitTestProcess> InPro
 	{
 		OnSuspendStateChange.ExecuteIfBound(InSuspendState);
 	}
-
-	// @todo #JohnBLowPri: Want to support client process debugging in future?
 }
 
 
@@ -719,36 +636,16 @@ bool UClientUnitTest::NotifyConsoleCommandRequest(FString CommandContext, FStrin
 			//				then it leads to an early disconnect when you try to spam-send a command early.
 			//
 			//				It's easy to test this, just type in a command before join, and hold down enter on the edit box to spam it.
-			if (UnitConn != NULL)
-			{
-				FOutBunch* ControlChanBunch = NUTNet::CreateChannelBunch(ControlBunchSequence, UnitConn, CHTYPE_Control, 0);
-
-				if (ControlChanBunch != nullptr)
-				{
-					uint8 ControlMsg = NMT_NUTControl;
-					uint8 ControlCmd = ENUTControlCommand::Command_NoResult;
-					FString Cmd = Command;
-
-					*ControlChanBunch << ControlMsg;
-					*ControlChanBunch << ControlCmd;
-					*ControlChanBunch << Cmd;
-
-					NUTNet::SendControlBunch(UnitConn, *ControlChanBunch);
-
 
-					UNIT_LOG(ELogType::OriginConsole, TEXT("Sent command '%s' to server."), *Command);
+			if (SendNUTControl(ENUTControlCommand::Command_NoResult, Command))
+			{
+				UNIT_LOG(ELogType::OriginConsole, TEXT("Sent command '%s' to server."), *Command);
 
-					bHandled = true;
-				}
-				else
-				{
-					UNIT_LOG(ELogType::OriginConsole, TEXT("Failed to send console command '%s', failed to create control bunch."),
-								*Command);
-				}
+				bHandled = true;
 			}
 			else
 			{
-				UNIT_LOG(ELogType::OriginConsole, TEXT("Failed to send console command '%s', no server connection."), *Command);
+				UNIT_LOG(ELogType::OriginConsole, TEXT("Failed to send console command '%s' to server."), *Command);
 			}
 		}
 		else if (CommandContext == TEXT("Client"))
@@ -782,6 +679,34 @@ void UClientUnitTest::GetCommandContextList(TArray<TSharedPtr<FString>>& OutList
 }
 
 
+bool UClientUnitTest::SendNUTControl(ENUTControlCommand CommandType, FString Command)
+{
+	bool bSuccess = false;
+
+	FOutBunch* ControlChanBunch = MinClient->CreateChannelBunch(CHTYPE_Control, 0);
+
+	if (ControlChanBunch != nullptr)
+	{
+		uint8 ControlMsg = NMT_NUTControl;
+		ENUTControlCommand CmdType = CommandType;
+
+		*ControlChanBunch << ControlMsg;
+		*ControlChanBunch << CmdType;
+		*ControlChanBunch << Command;
+
+		bSuccess = MinClient->SendControlBunch(ControlChanBunch);
+	}
+	else
+	{
+		FString LogMsg = TEXT("Failed to create control channel bunch.");
+
+		UNIT_LOG(ELogType::StatusFailure, TEXT("%s"), *LogMsg);
+		UNIT_STATUS_LOG(ELogType::StatusVerbose, TEXT("%s"), *LogMsg);
+	}
+
+	return bSuccess;
+}
+
 bool UClientUnitTest::SendRPCChecked(UObject* Target, const TCHAR* FunctionName, void* Parms, int16 ParmsSize,
 										int16 ParmsSizeCorrection/*=0*/)
 {
@@ -794,7 +719,7 @@ bool UClientUnitTest::SendRPCChecked(UObject* Target, const TCHAR* FunctionName,
 	{
 		if (TargetFunc->ParmsSize == ParmsSize + ParmsSizeCorrection)
 		{
-			if (UnitConn->IsNetReady(false))
+			if (MinClient->GetConn()->IsNetReady(false))
 			{
 				Target->ProcessEvent(TargetFunc, Parms);
 			}
@@ -829,7 +754,7 @@ bool UClientUnitTest::SendRPCChecked(UObject* Target, FFuncReflection& FuncRefl)
 	}
 	else
 	{
-		UNIT_LOG(ELogType::StatusFailure, TEXT("Failed to send RPC '%s', reflection failed."), FuncRefl.FunctionName);
+		UNIT_LOG(ELogType::StatusFailure, TEXT("Failed to send RPC '%s', function reflection failed."), FuncRefl.FunctionName);
 	}
 
 	return bSuccess;
@@ -861,16 +786,17 @@ bool UClientUnitTest::SendUnitRPCChecked_Internal(UObject* Target, FString RPCNa
 void UClientUnitTest::PreSendRPC()
 {
 	// Flush before and after, so no queued data is counted as a send, and so that the queued RPC is immediately sent and detected
-	UnitConn->FlushNet();
+	MinClient->GetConn()->FlushNet();
 
-	GSentBunch = false;
+	bSentBunch = false;
 }
 
-bool UClientUnitTest::PostSendRPC(FString RPCName, UObject* Target/*=NULL*/)
+bool UClientUnitTest::PostSendRPC(FString RPCName, UObject* Target/*=nullptr*/)
 {
 	bool bSuccess = false;
 	UActorComponent* TargetComponent = Cast<UActorComponent>(Target);
 	AActor* TargetActor = (TargetComponent != nullptr ? TargetComponent->GetOwner() : Cast<AActor>(Target));
+	UNetConnection* UnitConn = MinClient->GetConn();
 	UChannel* TargetChan = UnitConn->ActorChannels.FindRef(TargetActor);
 
 	UnitConn->FlushNet();
@@ -882,7 +808,7 @@ bool UClientUnitTest::PostSendRPC(FString RPCName, UObject* Target/*=NULL*/)
 	}
 
 	// If sending failed, trigger an overall unit test failure
-	if (!GSentBunch)
+	if (!bSentBunch)
 	{
 		FString LogMsg = FString::Printf(TEXT("Failed to send RPC '%s', unit test needs update."), *RPCName);
 
@@ -987,30 +913,17 @@ bool UClientUnitTest::PostSendRPC(FString RPCName, UObject* Target/*=NULL*/)
 
 void UClientUnitTest::SendGenericExploitFailLog()
 {
-	FOutBunch* ControlChanBunch = NUTNet::CreateChannelBunch(ControlBunchSequence, UnitConn, CHTYPE_Control, 0);
-
-	if (ControlChanBunch != nullptr)
-	{
-		uint8 ControlMsg = NMT_NUTControl;
-		uint8 ControlCmd = ENUTControlCommand::Command_NoResult;
-		FString Cmd = GetGenericExploitFailLog();
-
-		*ControlChanBunch << ControlMsg;
-		*ControlChanBunch << ControlCmd;
-		*ControlChanBunch << Cmd;
-
-		NUTNet::SendControlBunch(UnitConn, *ControlChanBunch);
-	}
+	SendNUTControl(ENUTControlCommand::Command_NoResult, GetGenericExploitFailLog());
 }
 
-
-// @todo #JohnBRefactor: This should be changed to work as compile-time conditionals, if possible
 bool UClientUnitTest::ValidateUnitTestSettings(bool bCDOCheck/*=false*/)
 {
 	bool bSuccess = Super::ValidateUnitTestSettings();
 
-	// Currently, unit tests don't support NOT launching/connecting to a server
-	UNIT_ASSERT(!!(UnitTestFlags & EUnitTestFlags::LaunchServer));
+	ValidateUnitFlags(UnitTestFlags, MinClientFlags);
+
+
+	// Validate the rest of the flags which cross-check against non-flag variables
 
 	// If launching a server, make sure the base URL for the server is set
 	UNIT_ASSERT(!(UnitTestFlags & EUnitTestFlags::LaunchServer) || BaseServerURL.Len() > 0);
@@ -1020,34 +933,11 @@ bool UClientUnitTest::ValidateUnitTestSettings(bool bCDOCheck/*=false*/)
 
 
 	// You can't specify an allowed actors whitelist, without the AcceptActors flag
-	UNIT_ASSERT(AllowedClientActors.Num() == 0 || !!(UnitTestFlags & EUnitTestFlags::AcceptActors));
-
-	// If you require a player/NUTActor, you need to accept actor channels
-	UNIT_ASSERT((!(UnitTestFlags & EUnitTestFlags::AcceptPlayerController) && !(UnitTestFlags & EUnitTestFlags::RequireNUTActor)) ||
-					!!(UnitTestFlags & EUnitTestFlags::AcceptActors));
-
-	// Don't require a PlayerController, if you don't accept one
-	UNIT_ASSERT(!(UnitTestFlags & EUnitTestFlags::RequirePlayerController) ||
-					!!(UnitTestFlags & EUnitTestFlags::AcceptPlayerController));
-
-	// If you require a pawn, you must require a PlayerController
-	UNIT_ASSERT(!(UnitTestFlags & EUnitTestFlags::RequirePawn) || !!(UnitTestFlags & EUnitTestFlags::RequirePlayerController));
-
-	// If you require a pawn, you must enable NotifyProcessEvent
-	UNIT_ASSERT(!(UnitTestFlags & EUnitTestFlags::RequirePawn) || !!(UnitTestFlags & EUnitTestFlags::NotifyProcessEvent));
-
-	// If you require a PlayerState, you must require a PlayerController
-	UNIT_ASSERT(!(UnitTestFlags & EUnitTestFlags::RequirePlayerState) || !!(UnitTestFlags & EUnitTestFlags::RequirePlayerController));
-
-	// If you want to dump received RPC's, you need to hook NotifyProcessEvent
-	UNIT_ASSERT(!(UnitTestFlags & EUnitTestFlags::DumpReceivedRPC) || !!(UnitTestFlags & EUnitTestFlags::NotifyProcessEvent));
-
-	// You can't whitelist client RPC's (i.e. unblock whitelisted RPC's), unless all RPC's are blocked by default
-	UNIT_ASSERT(!(UnitTestFlags & EUnitTestFlags::AcceptRPCs) || AllowedClientRPCs.Num() == 0);
+	UNIT_ASSERT(AllowedClientActors.Num() == 0 || !!(MinClientFlags & EMinClientFlags::AcceptActors));
 
 #if UE_BUILD_SHIPPING
 	// You can't hook ProcessEvent or block RPCs in shipping builds, as the main engine hook is not available in shipping; soft-fail
-	if (!!(UnitTestFlags & EUnitTestFlags::NotifyProcessEvent) || !(UnitTestFlags & EUnitTestFlags::AcceptRPCs))
+	if (!!(UnitTestFlags & EUnitTestFlags::NotifyProcessEvent) || !(MinClientFlags & EMinClientFlags::NotifyProcessNetEvent))
 	{
 		UNIT_LOG(ELogType::StatusFailure | ELogType::StyleBold, TEXT("Unit tests run in shipping mode, can't hook ProcessEvent."));
 
@@ -1055,71 +945,24 @@ bool UClientUnitTest::ValidateUnitTestSettings(bool bCDOCheck/*=false*/)
 	}
 #endif
 
-	// If you require a pawn, validate the existence of certain RPC's that are needed for pawn setup and verification
-	UNIT_ASSERT(!(UnitTestFlags & EUnitTestFlags::RequirePawn) || (
-				GetDefault<APlayerController>()->FindFunction(FName(TEXT("ClientRestart"))) != NULL &&
-				GetDefault<APlayerController>()->FindFunction(FName(TEXT("ClientRetryClientRestart"))) != NULL));
-
-	// For part of pawn-setup detection, you need notification for net actors
-	UNIT_ASSERT(!(UnitTestFlags & EUnitTestFlags::RequirePawn) || !!(UnitTestFlags & EUnitTestFlags::NotifyNetActors));
-
-	// For part of PlayerState-setup detection, you need notification for net actors
-	UNIT_ASSERT(!(UnitTestFlags & EUnitTestFlags::RequirePlayerState) || !!(UnitTestFlags & EUnitTestFlags::NotifyNetActors));
-
 	// If the ping requirements flag is set, it should be the ONLY one set
-	// (which means only one bit should be set, and one bit means it should be power-of-two)
+	//	(which means only one bit should be set, and one bit means it should be power-of-two)
 	UNIT_ASSERT(!(UnitTestFlags & EUnitTestFlags::RequirePing) ||
 					FMath::IsPowerOfTwo((uint32)(UnitTestFlags & EUnitTestFlags::RequirementsMask)));
 
-	// As above, but with the 'custom' requirements flag
-	// NOTE: Removed, as it can still be useful to have other requirements flags
-	//UNIT_ASSERT(!(UnitTestFlags & EUnitTestFlags::RequireCustom) ||
-	//				FMath::IsPowerOfTwo((uint32)(UnitTestFlags & EUnitTestFlags::RequirementsMask)));
-
-	// You can't accept RPC's, without accepting actors
-	UNIT_ASSERT(!(UnitTestFlags & EUnitTestFlags::AcceptRPCs) || !!(UnitTestFlags & EUnitTestFlags::AcceptActors));
-
-	// You can't send RPC's, without accepting a player controller (netcode blocks this, without a PC); unless this is a beacon
-	UNIT_ASSERT(!(UnitTestFlags & EUnitTestFlags::SendRPCs) || !!(UnitTestFlags & EUnitTestFlags::AcceptPlayerController) ||
-				!!(UnitTestFlags & EUnitTestFlags::BeaconConnect));
-
-	// If connecting to a beacon, a number of unit test flags are not supported
-	const EUnitTestFlags RejectedBeaconFlags = (EUnitTestFlags::AcceptPlayerController | EUnitTestFlags::RequirePlayerController |
-												EUnitTestFlags::RequirePing);
-
-	UNIT_ASSERT(!(UnitTestFlags & EUnitTestFlags::BeaconConnect) || !(UnitTestFlags & RejectedBeaconFlags));
-
-	// If connecting to a beacon, net actor notification is required, for proper setup
-	UNIT_ASSERT(!(UnitTestFlags & EUnitTestFlags::BeaconConnect) || !!(UnitTestFlags & EUnitTestFlags::NotifyNetActors));
+	// If you require a pawn, validate the existence of certain RPC's that are needed for pawn setup and verification
+	UNIT_ASSERT(!(UnitTestFlags & EUnitTestFlags::RequirePawn) || (
+				GetDefault<APlayerController>()->FindFunction(FName(TEXT("ClientRestart"))) != nullptr &&
+				GetDefault<APlayerController>()->FindFunction(FName(TEXT("ClientRetryClientRestart"))) != nullptr));
 
 	// If connecting to a beacon, you must specify the beacon type
 	UNIT_ASSERT(!(UnitTestFlags & EUnitTestFlags::BeaconConnect) || ServerBeaconType.Len() > 0);
 
-	// Don't require a beacon, if you're not connecting to a beacon
-	UNIT_ASSERT(!(UnitTestFlags & EUnitTestFlags::RequireBeacon) || !!(UnitTestFlags & EUnitTestFlags::BeaconConnect));
-
-	// Don't specify server-dependent flags, if not auto-launching a server
-	UNIT_ASSERT(!(UnitTestFlags & EUnitTestFlags::LaunchClient) || !!(UnitTestFlags & EUnitTestFlags::LaunchServer));
-
-	// If DumpReceivedRaw is set, make sure hooking of ReceivedRawPacket is enabled too
-	UNIT_ASSERT(!(UnitTestFlags & EUnitTestFlags::DumpReceivedRaw) || !!(UnitTestFlags & EUnitTestFlags::CaptureReceivedRaw));
-
-	// If DumpSendRaw is set, make sure hooking of SendRawPacket is enabled too
-	UNIT_ASSERT(!(UnitTestFlags & EUnitTestFlags::DumpSendRaw) || !!(UnitTestFlags & EUnitTestFlags::CaptureSendRaw));
-
-	// You can't get net actor notifications, unless you accept actors
-	UNIT_ASSERT(!(UnitTestFlags & EUnitTestFlags::NotifyNetActors) || !!(UnitTestFlags & EUnitTestFlags::AcceptActors));
-
-	// You can't use 'RequireNUTActor', without net actor notifications
-	UNIT_ASSERT(!(UnitTestFlags & EUnitTestFlags::RequireNUTActor) || !!(UnitTestFlags & EUnitTestFlags::NotifyNetActors));
 
 	// Don't accept any 'Ignore' flags, once the unit test is finalized (they're debug only - all crashes must be handled in final code)
 	UNIT_ASSERT(bWorkInProgress || !(UnitTestFlags & (EUnitTestFlags::IgnoreServerCrash | EUnitTestFlags::IgnoreClientCrash |
 				EUnitTestFlags::IgnoreDisconnect)));
 
-	// If a unit test expects a server crash, it should also expect a disconnect too (to avoid an invalid 'needs update' result)
-	UNIT_ASSERT(!(UnitTestFlags & EUnitTestFlags::ExpectServerCrash) || !!(UnitTestFlags & EUnitTestFlags::ExpectDisconnect));
-
 
 	return bSuccess;
 }
@@ -1178,7 +1021,7 @@ bool UClientUnitTest::HasAllRequirements(bool bIgnoreCustom/*=false*/)
 	bool bReturnVal = true;
 
 	// The fake client creation/connection is now delayed, so need to wait for that too
-	if (UnitConn == nullptr || MinClient == nullptr || !MinClient->IsConnected())
+	if (MinClient == nullptr || !MinClient->IsConnected())
 	{
 		bReturnVal = false;
 	}
@@ -1212,7 +1055,12 @@ ELogType UClientUnitTest::GetExpectedLogTypes()
 		ReturnVal |= ELogType::Client;
 	}
 
-	if (!!(UnitTestFlags & (EUnitTestFlags::DumpReceivedRaw | EUnitTestFlags::DumpSendRaw | EUnitTestFlags::DumpControlMessages)))
+	if (!!(UnitTestFlags & EUnitTestFlags::DumpControlMessages))
+	{
+		ReturnVal |= ELogType::StatusDebug;
+	}
+
+	if (!!(MinClientFlags & (EMinClientFlags::DumpReceivedRaw | EMinClientFlags::DumpSendRaw)))
 	{
 		ReturnVal |= ELogType::StatusDebug;
 	}
@@ -1240,12 +1088,15 @@ void UClientUnitTest::ResetTimeout(FString ResetReason, bool bResetConnTimeout/*
 
 void UClientUnitTest::ResetConnTimeout(float Duration)
 {
-	if (UnitConn != NULL && UnitConn->State != USOCK_Closed && UnitConn->Driver != NULL)
+	UNetConnection* UnitConn = (MinClient != nullptr ? MinClient->GetConn() : nullptr);
+	UNetDriver* UnitDriver = (UnitConn != nullptr ? UnitConn->Driver : nullptr);
+
+	if (UnitConn != nullptr && UnitConn->State != USOCK_Closed && UnitDriver != nullptr)
 	{
 		// @todo #JohnBHack: This is a slightly hacky way of setting the timeout to a large value, which will be overridden by newly
 		//				received packets, making it unsuitable for most situations (except crashes - but that could still be subject
 		//				to a race condition)
-		double NewLastReceiveTime = UnitConn->Driver->Time + Duration;
+		double NewLastReceiveTime = UnitDriver->Time + Duration;
 
 		UnitConn->LastReceiveTime = FMath::Max(NewLastReceiveTime, UnitConn->LastReceiveTime);
 	}
@@ -1254,10 +1105,9 @@ void UClientUnitTest::ResetConnTimeout(float Duration)
 
 bool UClientUnitTest::ExecuteUnitTest()
 {
-	bool bSuccess = false;
-	bool bValidSettings = ValidateUnitTestSettings();
+	bool bSuccess = ValidateUnitTestSettings();
 
-	if (bValidSettings)
+	if (bSuccess)
 	{
 		if (!!(UnitTestFlags & EUnitTestFlags::LaunchServer))
 		{
@@ -1295,16 +1145,8 @@ bool UClientUnitTest::ExecuteUnitTest()
 				}
 			}
 		}
-
-		// No longer immediately setup the fake client, wait for the server to fully startup first (monitoring its log output)
-#if 0
-		bSuccess = ConnectFakeClient();
-		bTriggerredInitialConnect = true;
-#else
-		bSuccess = true;
-#endif
 	}
-	else if (!bValidSettings)
+	else
 	{
 		FString LogMsg = TEXT("Failed to validate unit test settings/environment");
 
@@ -1317,16 +1159,17 @@ bool UClientUnitTest::ExecuteUnitTest()
 
 void UClientUnitTest::CleanupUnitTest()
 {
-	CleanupFakeClient();
+	if (MinClient != nullptr)
+	{
+		FProcessEventHook::Get().RemoveEventHook(MinClient->GetUnitWorld());
+	}
 
-#if !UE_BUILD_SHIPPING
-	RemoveProcessEventCallback(this, &UClientUnitTest::InternalScriptProcessEvent);
-#endif
+	CleanupMinimalClient();
 
 	Super::CleanupUnitTest();
 }
 
-bool UClientUnitTest::ConnectFakeClient(FUniqueNetIdRepl* InNetID/*=nullptr*/)
+bool UClientUnitTest::ConnectMinimalClient(const TCHAR* InNetID/*=nullptr*/)
 {
 	bool bSuccess = false;
 	FMinClientHooks Hooks;
@@ -1335,105 +1178,74 @@ bool UClientUnitTest::ConnectFakeClient(FUniqueNetIdRepl* InNetID/*=nullptr*/)
 
 	Hooks.ConnectedDel = FOnMinClientConnected::CreateUObject(this, &UClientUnitTest::NotifyMinClientConnected);
 	Hooks.NetworkFailureDel = FOnMinClientNetworkFailure::CreateUObject(this, &UClientUnitTest::NotifyNetworkFailure);
-	Hooks.SendRPCDel = FOnMinClientSendRPC::CreateUObject(this, &UClientUnitTest::NotifySendRPC);
+#if !UE_BUILD_SHIPPING
+	Hooks.SendRPCDel = FOnSendRPC::CreateUObject(this, &UClientUnitTest::NotifySendRPC);
+#endif
 	Hooks.ReceivedControlBunchDel = FOnMinClientReceivedControlBunch::CreateUObject(this, &UClientUnitTest::ReceivedControlBunch);
 	Hooks.RepActorSpawnDel = FOnMinClientRepActorSpawn::CreateUObject(this, &UClientUnitTest::NotifyAllowNetActor);
+	Hooks.HandleClientPlayerDel = FOnHandleClientPlayer::CreateUObject(this, &UClientUnitTest::NotifyHandleClientPlayer);
 
-	// @todo #JohnB: Restrict based on a new 'NotifyFlags' field, as this is presently never used by anything
 	if (!!(UnitTestFlags & EUnitTestFlags::CaptureReceivedRaw))
 	{
 		Hooks.ReceivedRawPacketDel = FOnMinClientReceivedRawPacket::CreateUObject(this, &UClientUnitTest::NotifyReceivedRawPacket);
 	}
 
-	if (!!(UnitTestFlags & EUnitTestFlags::CaptureSendRaw))
-	{
-		Hooks.LowLevelSendDel.BindUObject(this, &UClientUnitTest::NotifySendRawPacket);
-		Hooks.SocketSendDel.BindUObject(this, &UClientUnitTest::NotifySocketSendRawPacket);
-	}
+#if !UE_BUILD_SHIPPING
+	Hooks.LowLevelSendDel.BindUObject(this, &UClientUnitTest::NotifySocketSendRawPacket);
+#endif
 
 
-	EMinClientFlags MinClientFlags = FromUnitTestFlags(UnitTestFlags);
+	EMinClientFlags CurMinClientFlags = FromUnitTestFlags(UnitTestFlags) | MinClientFlags;
 
-	if (!!(MinClientFlags & EMinClientFlags::NotifyNetActors))
+	if (!!(CurMinClientFlags & EMinClientFlags::NotifyNetActors))
 	{
 		Hooks.NetActorDel.BindUObject(this, &UClientUnitTest::NotifyNetActor);
 	}
 
-	// @todo #JohnB: If possible, eliminate this interdependency with minimal clients, so that they can run independently
-	bool bFailedScriptHook = false;
 
-	if (!(MinClientFlags & EMinClientFlags::AcceptRPCs) || !!(MinClientFlags & EMinClientFlags::NotifyProcessEvent))
-	{
-#if !UE_BUILD_SHIPPING
-		AddProcessEventCallback(this, &UClientUnitTest::InternalScriptProcessEvent);
-		bFailedScriptHook = false;
-#else
-		bFailedScriptHook = true;
-#endif
-	}
+	Hooks.ReceiveRPCDel.BindUObject(this, &UClientUnitTest::NotifyReceiveRPC);
 
-	if (!bFailedScriptHook)
-	{
-		FMinClientParms Parms;
 
-		Parms.MinClientFlags = MinClientFlags;
-		Parms.Owner = this;
-		Parms.ServerAddress = ServerAddress;
-		Parms.BeaconAddress = BeaconAddress;
-		Parms.BeaconType = ServerBeaconType;
-		Parms.JoinUID = InNetID;
-
-		MinClient = NewObject<UMinimalClient>();
-		bSuccess = MinClient->Connect(Parms, Hooks);
+	FMinClientParms Parms;
 
-		if (bSuccess)
-		{
-			bTriggerredInitialConnect = true;
+	Parms.MinClientFlags = CurMinClientFlags;
+	Parms.Owner = this;
+	Parms.ServerAddress = ServerAddress;
+	Parms.BeaconAddress = BeaconAddress;
+	Parms.BeaconType = ServerBeaconType;
 
-			if (HasAllRequirements())
-			{
-				ResetTimeout(TEXT("ExecuteClientUnitTest (ExecuteUnitTest)"));
-				ExecuteClientUnitTest();
-			}
-
-			if (!!(UnitTestFlags & EUnitTestFlags::RequirePing))
-			{
-				// Send the 'ping' bunch
-				*ControlBunchSequence = 1;
-				FOutBunch* ControlChanBunch = NUTNet::CreateChannelBunch(ControlBunchSequence, UnitConn, CHTYPE_Control, 0);
-
-				if (ControlChanBunch != nullptr)
-				{
-					uint8 ControlMsg = NMT_NUTControl;
-					uint8 PingCmd = ENUTControlCommand::Ping;
-					FString Dud = TEXT("");
+	if (InNetID != nullptr)
+	{
+		Parms.JoinUID = InNetID;
+	}
 
-					*ControlChanBunch << ControlMsg;
-					*ControlChanBunch << PingCmd;
-					*ControlChanBunch << Dud;
+	Parms.AllowedClientRPCs = AllowedClientRPCs;
 
-					NUTNet::SendControlBunch(UnitConn, *ControlChanBunch);
-				}
-				else
-				{
-					FString LogMsg = TEXT("Failed to create control channel bunch.");
+	MinClient = NewObject<UMinimalClient>();
+	bSuccess = MinClient->Connect(Parms, Hooks);
 
-					UNIT_LOG(ELogType::StatusFailure, TEXT("%s"), *LogMsg);
-					UNIT_STATUS_LOG(ELogType::StatusVerbose, TEXT("%s"), *LogMsg);
-				}
-			}
-		}
-		else
+	if (bSuccess)
+	{
+		if (!!(UnitTestFlags & EUnitTestFlags::NotifyProcessEvent))
 		{
-			FString LogMsg = TEXT("Failed to connect minimal client.");
+#if !UE_BUILD_SHIPPING
+			FProcessEventHook::Get().AddEventHook(MinClient->GetUnitWorld(),
+				FOnProcessNetEvent::CreateUObject(this, &UClientUnitTest::NotifyProcessEvent));
+#else
+			FString LogMsg = TEXT("Require ProcessEvent hook, but current build configuration does not support it.");
 
 			UNIT_LOG(ELogType::StatusFailure, TEXT("%s"), *LogMsg);
 			UNIT_STATUS_LOG(ELogType::StatusVerbose, TEXT("%s"), *LogMsg);
+
+			bSuccess = false;
+#endif
 		}
+
+		bTriggerredInitialConnect = true;
 	}
 	else
 	{
-		FString LogMsg = TEXT("Require ProcessEvent hook, but current build configuration does not support it.");
+		FString LogMsg = TEXT("Failed to connect minimal client.");
 
 		UNIT_LOG(ELogType::StatusFailure, TEXT("%s"), *LogMsg);
 		UNIT_STATUS_LOG(ELogType::StatusVerbose, TEXT("%s"), *LogMsg);
@@ -1442,27 +1254,20 @@ bool UClientUnitTest::ConnectFakeClient(FUniqueNetIdRepl* InNetID/*=nullptr*/)
 	return bSuccess;
 }
 
-// @todo #JohnBRefactor: Reconsider naming of this function, if you modify the other functions using 'fake'
-void UClientUnitTest::CleanupFakeClient()
+void UClientUnitTest::CleanupMinimalClient()
 {
 	if (MinClient != nullptr)
 	{
 		MinClient->Cleanup();
 	}
 
-
 	UnitPC = nullptr;
-
 	bUnitPawnSetup = false;
 	bUnitPlayerStateSetup = false;
 	UnitNUTActor = nullptr;
 	bUnitNUTActorSetup = false;
 	UnitBeacon = nullptr;
 	bReceivedPong = false;
-
-	// @todo #JohnB: Remove after migrating
-	ControlBunchSequence = 0;
-
 	bPendingNetworkFailure = false;
 }
 
@@ -1470,8 +1275,8 @@ void UClientUnitTest::TriggerAutoReconnect()
 {
 	UNIT_LOG(ELogType::StatusImportant, TEXT("Performing Auto-Reconnect."))
 
-	CleanupFakeClient();
-	ConnectFakeClient();
+	CleanupMinimalClient();
+	ConnectMinimalClient();
 }
 
 
@@ -1560,7 +1365,7 @@ FString UClientUnitTest::ConstructServerParameters()
 
 TWeakPtr<FUnitTestProcess> UClientUnitTest::StartUnitTestClient(FString ConnectIP, bool bMinimized/*=true*/)
 {
-	TWeakPtr<FUnitTestProcess> ReturnVal = NULL;
+	TWeakPtr<FUnitTestProcess> ReturnVal = nullptr;
 
 	FString LogMsg = TEXT("Unit test launching a client");
 
@@ -1632,64 +1437,6 @@ void UClientUnitTest::PrintUnitTestProcessErrors(TSharedPtr<FUnitTestProcess> In
 	Super::PrintUnitTestProcessErrors(InHandle);
 }
 
-
-#if !UE_BUILD_SHIPPING
-// @todo #JohnB: 'HookOrigin' is not optimal - try to at least make it a UClientUnitTest pointer
-bool UClientUnitTest::InternalScriptProcessEvent(AActor* Actor, UFunction* Function, void* Parameters, void* HookOrigin)
-{
-	bool bBlockEvent = false;
-
-	if (HookOrigin != NULL)
-	{
-		// Match up the origin to an active unit test (for safety)
-		UClientUnitTest* OriginUnitTest = NULL;
-
-		if (GUnitTestManager != NULL)
-		{
-			for (int i=0; i<GUnitTestManager->ActiveUnitTests.Num(); i++)
-			{
-				UUnitTest* CurUnitTest = GUnitTestManager->ActiveUnitTests[i];
-
-				if (CurUnitTest == HookOrigin)
-				{
-					OriginUnitTest = Cast<UClientUnitTest>(CurUnitTest);
-
-					// @todo #JohnB: Restore UnitNetDriver check?
-					if (/*OriginUnitTest->UnitNetDriver == NULL ||*/ OriginUnitTest->UnitConn == NULL)
-					{
-						OriginUnitTest = NULL;
-					}
-
-					break;
-				}
-			}
-		}
-
-
-		if (OriginUnitTest != NULL)
-		{
-			// Verify the net driver matches, so we don't stomp on other unit test actors
-			if (Actor != NULL && ((Actor->GetWorld() != NULL && Actor->GetNetDriver() == OriginUnitTest->UnitConn->Driver) ||
-				OriginUnitTest->UnitConn->ActorChannels.Contains(Actor)))
-
-			{
-				UNIT_EVENT_BEGIN(OriginUnitTest);
-
-				bBlockEvent = OriginUnitTest->NotifyScriptProcessEvent(Actor, Function, Parameters);
-
-				UNIT_EVENT_END;
-			}
-		}
-		else
-		{
-			UE_LOG(LogUnitTest, Warning, TEXT("Failed to find HookOrigin in active unit test list"));
-		}
-	}
-
-	return bBlockEvent;
-}
-#endif
-
 void UClientUnitTest::UnitTick(float DeltaTime)
 {
 	if (bBlockingServerDelay || bBlockingClientDelay || bBlockingFakeClientDelay)
@@ -1724,7 +1471,7 @@ void UClientUnitTest::UnitTick(float DeltaTime)
 
 			if (bBlockingFakeClientDelay && !IsWaitingOnTimeout())
 			{
-				ConnectFakeClient();
+				ConnectMinimalClient();
 
 				bTriggerredInitialConnect = true;
 				bBlockingFakeClientDelay = false;
diff --git a/Engine/Plugins/NetcodeUnitTest/NetcodeUnitTest/Source/NetcodeUnitTest/Private/MinimalClient.cpp b/Engine/Plugins/NetcodeUnitTest/NetcodeUnitTest/Source/NetcodeUnitTest/Private/MinimalClient.cpp
--- a/Engine/Plugins/NetcodeUnitTest/NetcodeUnitTest/Source/NetcodeUnitTest/Private/MinimalClient.cpp
+++ b/Engine/Plugins/NetcodeUnitTest/NetcodeUnitTest/Source/NetcodeUnitTest/Private/MinimalClient.cpp
@@ -5,23 +5,46 @@
 #include "Misc/FeedbackContext.h"
 #include "Misc/NetworkVersion.h"
 #include "Engine/Engine.h"
+#include "Engine/GameEngine.h"
 #include "Engine/ActorChannel.h"
 #include "Net/DataChannel.h"
-#include "OnlineSubsystemTypes.h"
 
+#include "ClientUnitTest.h"
 #include "NetcodeUnitTest.h"
 #include "UnitTestEnvironment.h"
 #include "UnitTestManager.h"
+#include "NUTUtilDebug.h"
+#include "NUTUtilReflection.h"
 #include "Net/NUTUtilNet.h"
-#include "Net/UnitTestNetDriver.h"
-#include "Net/UnitTestNetConnection.h"
+#include "Net/UnitTestPackageMap.h"
 #include "Net/UnitTestChannel.h"
+#include "Net/UnitTestActorChannel.h"
 
 
 // @todo #JohnB: Eventually, you want to remove this classes dependence on the unit test code, and have it be independent.
 //					For now though, just allow the dependence, to ease the transition
 
 
+/**
+ * FMinClientParms
+ */
+void FMinClientParms::ValidateParms()
+{
+	ValidateMinFlags(MinClientFlags);
+
+
+	// Validate the rest of the flags which cross-check against non-flag variables, or otherwise should be runtime-only checks
+
+	// You can't whitelist client RPC's (i.e. unblock whitelisted RPC's), unless all RPC's are blocked by default
+	UNIT_ASSERT(!(MinClientFlags & EMinClientFlags::AcceptRPCs) || AllowedClientRPCs.Num() == 0);
+
+#if UE_BUILD_SHIPPING
+	// Rejecting actors requires non-shipping mode
+	UNIT_ASSERT(!!(MinClientFlags & EMinClientFlags::AcceptActors));
+#endif
+}
+
+
 /**
  * UMinimalClient
  */
@@ -34,7 +57,6 @@ UMinimalClient::UMinimalClient(const FObjectInitializer& ObjectInitializor)
 	, UnitWorld(nullptr)
 	, UnitNetDriver(nullptr)
 	, UnitConn(nullptr)
-	, ControlBunchSequence(0)
 	, PendingNetActorChans()
 {
 }
@@ -47,14 +69,11 @@ bool UMinimalClient::Connect(FMinClientParms Parms, FMinClientHooks Hooks)
 	check(Hooks.ConnectedDel.IsBound());
 	check(MinClientFlags == EMinClientFlags::None && !ConnectedDel.IsBound());
 
-	// @todo #JohnB: Add validation of minimal client flags, through another function, here.
+	Parms.ValidateParms();
 
 	Parms.CopyParms(*this);
 	Hooks.CopyHooks(*this);
 
-	// @todo #JohnB: Remove once fully migrated to minimal client
-	Owner->ControlBunchSequence = &ControlBunchSequence;
-
 	if (UnitWorld == nullptr)
 	{
 		// Make all of this happen in a blank, newly constructed world
@@ -64,8 +83,6 @@ bool UMinimalClient::Connect(FMinClientParms Parms, FMinClientHooks Hooks)
 		{
 			if (ConnectMinimalClient())
 			{
-				bSuccess = true;
-
 				if (GEngine != nullptr)
 				{
 #if TARGET_UE4_CL >= CL_DEPRECATEDEL
@@ -76,12 +93,13 @@ bool UMinimalClient::Connect(FMinClientParms Parms, FMinClientHooks Hooks)
 #endif
 				}
 
+				if (!(MinClientFlags & EMinClientFlags::AcceptRPCs) || !!(MinClientFlags & EMinClientFlags::NotifyProcessNetEvent))
+				{
+					FProcessEventHook::Get().AddRPCHook(UnitWorld,
+						FOnProcessNetEvent::CreateUObject(this, &UMinimalClient::NotifyReceiveRPC));
+				}
 
-				// @todo #JohnB: Deprecate this
-#if 1
-				Owner->UnitNetDriver = UnitNetDriver;
-				Owner->UnitConn = UnitConn;
-#endif
+				bSuccess = true;
 			}
 			else
 			{
@@ -125,20 +143,8 @@ void UMinimalClient::Cleanup()
 	}
 
 	UnitNetDriver = nullptr;
-
 	UnitConn = nullptr;
 
-	// @todo #JohnB: Deprecate this
-#if 1
-	Owner->UnitNetDriver = nullptr;
-	Owner->UnitConn = nullptr;
-#endif
-
-	ControlBunchSequence = 0;
-
-	// @todo #JohnB: Remove once fully migrated to minimal client
-	Owner->ControlBunchSequence = nullptr;
-
 	PendingNetActorChans.Empty();
 
 	if (GEngine != nullptr)
@@ -153,11 +159,99 @@ void UMinimalClient::Cleanup()
 	// Immediately cleanup (or rather, at start of next tick, as that's earliest possible time) after sending the RPC
 	if (UnitWorld != nullptr)
 	{
+		FProcessEventHook::Get().RemoveRPCHook(UnitWorld);
+
 		NUTNet::MarkUnitTestWorldForCleanup(UnitWorld);
 		UnitWorld = nullptr;
 	}
 }
 
+FOutBunch* UMinimalClient::CreateChannelBunch(EChannelType ChType, int32 ChIndex/*=INDEX_NONE*/)
+{
+	FOutBunch* ReturnVal = nullptr;
+	UChannel* ControlChan = UnitConn != nullptr ? UnitConn->Channels[0] : nullptr;
+
+	if (ControlChan != nullptr)
+	{
+		if (ChIndex == INDEX_NONE)
+		{
+			for (ChIndex=0; ChIndex<ARRAY_COUNT(UnitConn->Channels); ChIndex++)
+			{
+				if (UnitConn->Channels[ChIndex] == nullptr)
+				{
+					break;
+				}
+			}
+
+			check(UnitConn->Channels[ChIndex] == nullptr);
+		}
+
+		if (ControlChan != nullptr && ControlChan->IsNetReady(false))
+		{
+			int32 BunchSequence = ++UnitConn->OutReliable[ChIndex];
+
+			ReturnVal = new FOutBunch(ControlChan, false);
+
+			ReturnVal->Next = nullptr;
+			ReturnVal->Time = 0.0;
+			ReturnVal->ReceivedAck = false;
+			ReturnVal->PacketId = 0;
+			ReturnVal->bDormant = false;
+			ReturnVal->Channel = nullptr;
+			ReturnVal->ChIndex = ChIndex;
+			ReturnVal->ChType = ChType;
+			ReturnVal->bReliable = 1;
+			ReturnVal->ChSequence = BunchSequence;
+
+			// NOTE: Might not cover all bOpen or 'channel already open' cases
+			if (UnitConn->Channels[ChIndex] == nullptr)
+			{
+				ReturnVal->bOpen = 1;
+			}
+			else if (UnitConn->Channels[ChIndex]->OpenPacketId.First == INDEX_NONE)
+			{
+				ReturnVal->bOpen = 1;
+
+				UnitConn->Channels[ChIndex]->OpenPacketId.First = BunchSequence;
+				UnitConn->Channels[ChIndex]->OpenPacketId.Last = BunchSequence;
+			}
+		}
+	}
+
+	return ReturnVal;
+}
+
+bool UMinimalClient::SendControlBunch(FOutBunch* ControlChanBunch)
+{
+	bool bSuccess = false;
+
+	if (ControlChanBunch != nullptr && UnitConn != nullptr && UnitConn->Channels[0] != nullptr)
+	{
+		// Since it's the unit test control channel, sending the packet abnormally, append to OutRec manually
+		if (ControlChanBunch->bReliable)
+		{
+			UChannel* ControlChan = UnitConn->Channels[0];
+
+			for (FOutBunch* CurOut=ControlChan->OutRec; CurOut!=nullptr; CurOut=CurOut->Next)
+			{
+				if (CurOut->Next == nullptr)
+				{
+					CurOut->Next = ControlChanBunch;
+					ControlChan->NumOutRec++;
+
+					break;
+				}
+			}
+		}
+
+		UnitConn->SendRawBunch(*ControlChanBunch, true);
+
+		bSuccess = true;
+	}
+
+	return bSuccess;
+}
+
 void UMinimalClient::UnitTick(float DeltaTimm)
 {
 	if (!!(MinClientFlags & EMinClientFlags::NotifyNetActors) && PendingNetActorChans.Num() > 0 && UnitConn != nullptr)
@@ -185,44 +279,92 @@ bool UMinimalClient::ConnectMinimalClient()
 	bool bSuccess = false;
 
 	check(UnitWorld != nullptr);
-	check(UnitNetDriver == nullptr);
 
-	UnitNetDriver = NUTNet::CreateUnitTestNetDriver(UnitWorld);
+	CreateNetDriver();
 
 	if (UnitNetDriver != nullptr)
 	{
+		// Hack: Replace the control and actor channels, with stripped down unit test channels
+		UnitNetDriver->ChannelClasses[CHTYPE_Control] = UUnitTestChannel::StaticClass();
+		UnitNetDriver->ChannelClasses[CHTYPE_Actor] = UUnitTestActorChannel::StaticClass();
+
+		// @todo #JohnB: Block voice channel?
+
 		UnitNetDriver->InitialConnectTimeout = FMath::Max(UnitNetDriver->InitialConnectTimeout, (float)Owner->UnitTestTimeout);
 		UnitNetDriver->ConnectionTimeout = FMath::Max(UnitNetDriver->ConnectionTimeout, (float)Owner->UnitTestTimeout);
 
+#if !UE_BUILD_SHIPPING
 		if (!(MinClientFlags & EMinClientFlags::SendRPCs) || !!(MinClientFlags & EMinClientFlags::DumpSendRPC))
 		{
-			UUnitTestNetDriver* CurUnitDriver = CastChecked<UUnitTestNetDriver>(UnitNetDriver);
-
-			CurUnitDriver->SendRPCDel.BindUObject(this, &UMinimalClient::NotifySendRPC);
+			UnitNetDriver->SendRPCDel.BindUObject(this, &UMinimalClient::NotifySendRPC);
 		}
+#endif
 
 		bool bBeaconConnect = !!(MinClientFlags & EMinClientFlags::BeaconConnect);
 		FString ConnectAddress = (bBeaconConnect ? BeaconAddress : ServerAddress);
 		FURL DefaultURL;
 		FURL TravelURL(&DefaultURL, *ConnectAddress, TRAVEL_Absolute);
 		FString ConnectionError;
-		bSuccess = UnitNetDriver->InitConnect(this, TravelURL, ConnectionError);
+		UClass* NetConnClass = UnitNetDriver->NetConnectionClass;
+		UNetConnection* DefConn = NetConnClass != nullptr ? Cast<UNetConnection>(NetConnClass->GetDefaultObject()) : nullptr;
+
+		bSuccess = DefConn != nullptr;
+
+		if (bSuccess)
+		{
+			// Replace the package map class
+			TSubclassOf<UPackageMap> OldClass = DefConn->PackageMapClass;
+			UProperty* OldPostConstructLink = NetConnClass->PostConstructLink;
+			UProperty* PackageMapProp = FindFieldChecked<UProperty>(NetConnClass, TEXT("PackageMapClass"));
+
+			// Hack - force property initialization for the PackageMapClass property, so changing its default value works.
+			check(PackageMapProp != nullptr && PackageMapProp->PostConstructLinkNext == nullptr);
+
+			PackageMapProp->PostConstructLinkNext = NetConnClass->PostConstructLink;
+			NetConnClass->PostConstructLink = PackageMapProp;
+			DefConn->PackageMapClass = UUnitTestPackageMap::StaticClass();
+
+			bSuccess = UnitNetDriver->InitConnect(this, TravelURL, ConnectionError);
+
+			DefConn->PackageMapClass = OldClass;
+			NetConnClass->PostConstructLink = OldPostConstructLink;
+			PackageMapProp->PostConstructLinkNext = nullptr;
+		}
+		else
+		{
+			UNIT_LOG_OBJ(Owner, ELogType::StatusFailure, TEXT("Failed to replace PackageMapClass, minimal client connection failed."));
+		}
 
 		if (bSuccess)
 		{
 			UnitConn = UnitNetDriver->ServerConnection;
 
+			check(UnitConn->PackageMapClass == UUnitTestPackageMap::StaticClass());
+
 			FString LogMsg = FString::Printf(TEXT("Successfully created minimal client connection to IP '%s'"), *ConnectAddress);
 
 			UNIT_LOG_OBJ(Owner, ELogType::StatusImportant, TEXT("%s"), *LogMsg);
 			UNIT_STATUS_LOG_OBJ(Owner, ELogType::StatusVerbose, TEXT("%s"), *LogMsg);
 
 
-			UUnitTestNetConnection* CurUnitConn = CastChecked<UUnitTestNetConnection>(UnitConn);
+#if !UE_BUILD_SHIPPING
+			UnitConn->ReceivedRawPacketDel = FOnReceivedRawPacket::CreateUObject(this, &UMinimalClient::NotifyReceivedRawPacket);
+
+			if (!!(MinClientFlags & EMinClientFlags::DumpSendRaw))
+			{
+				UnitConn->LowLevelSendDel = FOnLowLevelSend::CreateUObject(this, &UMinimalClient::NotifySocketSend);
+			}
+			else
+			{
+				UnitConn->LowLevelSendDel = LowLevelSendDel;
+			}
+#endif
+
+			// Work around a minor UNetConnection bug, where QueuedBits is not initialized, until after the first Tick
+			UnitConn->QueuedBits = -(MAX_PACKET_SIZE * 8);
+
 			UUnitTestChannel* UnitControlChan = CastChecked<UUnitTestChannel>(UnitConn->Channels[0]);
 
-			CurUnitConn->MinClient = this;
-			CurUnitConn->SocketHook.MinClient = this;
 			UnitControlChan->MinClient = this;
 
 #if TARGET_UE4_CL >= CL_STATELESSCONNECT
@@ -239,23 +381,104 @@ bool UMinimalClient::ConnectMinimalClient()
 		}
 		else
 		{
-			UE_LOG(LogUnitTest, Error, TEXT("Failed to kickoff connect to IP '%s', error: %s"), *ConnectAddress, *ConnectionError);
+			UNIT_LOG_OBJ(Owner, ELogType::StatusFailure, TEXT("Failed to kickoff connect to IP '%s', error: %s"), *ConnectAddress,
+							*ConnectionError);
 		}
 	}
 	else
 	{
-		UE_LOG(LogUnitTest, Error, TEXT("Failed to create an instance of the unit test net driver"));
+		UNIT_LOG_OBJ(Owner, ELogType::StatusFailure, TEXT("Failed to create an instance of the unit test net driver"));
 	}
 
 	return bSuccess;
 }
 
-void UMinimalClient::SendInitialJoin()
+void UMinimalClient::CreateNetDriver()
 {
-	// @todo #JohnB: Deprecate ControlBunchSequence as soon as you can - and switch to using OutReliable on its own
-	ControlBunchSequence = UnitConn->OutReliable[0];
+	check(UnitNetDriver == nullptr);
+
+	UGameEngine* GameEngine = Cast<UGameEngine>(GEngine);
+
+	if (GameEngine != nullptr && UnitWorld != nullptr)
+	{
+		static int UnitTestNetDriverCount = 0;
 
-	FOutBunch* ControlChanBunch = NUTNet::CreateChannelBunch(ControlBunchSequence, UnitConn, CHTYPE_Control, 0);
+		// Setup a new driver name entry
+		bool bFoundDef = false;
+		FName UnitDefName = TEXT("UnitTestNetDriver");
+
+		for (int32 i=0; i<GameEngine->NetDriverDefinitions.Num(); i++)
+		{
+			if (GameEngine->NetDriverDefinitions[i].DefName == UnitDefName)
+			{
+				bFoundDef = true;
+				break;
+			}
+		}
+
+		if (!bFoundDef)
+		{
+			FNetDriverDefinition NewDriverEntry;
+
+			NewDriverEntry.DefName = UnitDefName;
+			NewDriverEntry.DriverClassName = TEXT("/Script/OnlineSubsystemUtils.IpNetDriver");
+			NewDriverEntry.DriverClassNameFallback = TEXT("/Script/OnlineSubsystemUtils.IpNetDriver");
+
+			GameEngine->NetDriverDefinitions.Add(NewDriverEntry);
+		}
+
+
+		FName NewDriverName = *FString::Printf(TEXT("UnitTestNetDriver_%i"), UnitTestNetDriverCount++);
+
+		// Now create a reference to the driver
+		if (GameEngine->CreateNamedNetDriver(UnitWorld, NewDriverName, UnitDefName))
+		{
+			UnitNetDriver = GameEngine->FindNamedNetDriver(UnitWorld, NewDriverName);
+		}
+
+
+		if (UnitNetDriver != nullptr)
+		{
+			UnitNetDriver->SetWorld(UnitWorld);
+			UnitWorld->SetNetDriver(UnitNetDriver);
+
+			UnitNetDriver->InitConnectionClass();
+
+
+			FLevelCollection* Collection = (FLevelCollection*)UnitWorld->GetActiveLevelCollection();
+
+			// Hack-set the net driver in the worlds level collection
+			if (Collection != nullptr)
+			{
+				Collection->SetNetDriver(UnitNetDriver);
+			}
+			else
+			{
+				UNIT_LOG_OBJ(Owner, ELogType::StatusWarning,
+								TEXT("CreateNetDriver: No LevelCollection found for created world, may block replication."));
+			}
+
+			UNIT_LOG_OBJ(Owner,, TEXT("CreateNetDriver: Created named net driver: %s, NetDriverName: %s, for World: %s"),
+							*UnitNetDriver->GetFullName(), *UnitNetDriver->NetDriverName.ToString(), *UnitWorld->GetFullName());
+		}
+		else
+		{
+			UNIT_LOG_OBJ(Owner, ELogType::StatusFailure, TEXT("CreateNetDriver: CreateNamedNetDriver failed"));
+		}
+	}
+	else if (GameEngine == nullptr)
+	{
+		UNIT_LOG_OBJ(Owner, ELogType::StatusFailure, TEXT("CreateNetDriver: GameEngine is nullptr"));
+	}
+	else //if (UnitWorld == nullptr)
+	{
+		UNIT_LOG_OBJ(Owner, ELogType::StatusFailure, TEXT("CreateNetDriver: UnitWorld is nullptr"));
+	}
+}
+
+void UMinimalClient::SendInitialJoin()
+{
+	FOutBunch* ControlChanBunch = CreateChannelBunch(CHTYPE_Control, 0);
 
 	if (ControlChanBunch != nullptr)
 	{
@@ -274,16 +497,9 @@ void UMinimalClient::SendInitialJoin()
 
 		bool bSkipControlJoin = !!(MinClientFlags & EMinClientFlags::SkipControlJoin);
 		bool bBeaconConnect = !!(MinClientFlags & EMinClientFlags::BeaconConnect);
-		TSharedPtr<const FUniqueNetId> DudPtr = MakeShareable(new FUniqueNetIdString(TEXT("Dud")));
-		FUniqueNetIdRepl PlayerUID(DudPtr);
 		FString BlankStr = TEXT("");
 		FString ConnectURL = UUnitTest::UnitEnv->GetDefaultClientConnectURL();
 
-		if (JoinUID != nullptr)
-		{
-			PlayerUID = *JoinUID;
-		}
-
 		if (bBeaconConnect)
 		{
 			if (!bSkipControlJoin)
@@ -291,7 +507,11 @@ void UMinimalClient::SendInitialJoin()
 				MessageType = NMT_BeaconJoin;
 				*ControlChanBunch << MessageType;
 				*ControlChanBunch << BeaconType;
-				*ControlChanBunch << PlayerUID;
+
+				int32 UIDSize = JoinUID.Len();
+
+				*ControlChanBunch << UIDSize;
+				*ControlChanBunch << JoinUID;
 
 				// Also immediately ack the beacon GUID setup; we're just going to let the server setup the client beacon,
 				// through the actor channel
@@ -307,7 +527,11 @@ void UMinimalClient::SendInitialJoin()
 			*ControlChanBunch << MessageType;
 			*ControlChanBunch << BlankStr;
 			*ControlChanBunch << ConnectURL;
-			*ControlChanBunch << PlayerUID;
+
+			int32 UIDSize = JoinUID.Len();
+
+			*ControlChanBunch << UIDSize;
+			*ControlChanBunch << JoinUID;
 
 
 			// Now send NMT_Join, to trigger a fake player, which should then trigger replication of basic actor channels
@@ -335,7 +559,7 @@ void UMinimalClient::SendInitialJoin()
 	}
 	else
 	{
-		UE_LOG(LogUnitTest, Error, TEXT("Failed to kickoff connection, could not create control channel bunch."));
+		UNIT_LOG_OBJ(Owner, ELogType::StatusFailure, TEXT("Failed to kickoff connection, could not create control channel bunch."));
 	}
 }
 
@@ -347,15 +571,119 @@ void UMinimalClient::DisconnectMinimalClient()
 	}
 }
 
-bool UMinimalClient::NotifySendRPC(AActor* Actor, UFunction* Function, void* Parameters, FOutParmRec* OutParms, FFrame* Stack,
-									UObject* SubObject)
+void UMinimalClient::NotifySocketSend(void* Data, int32 Count, bool& bBlockSend)
+{
+	if (!!(MinClientFlags & EMinClientFlags::DumpSendRaw))
+	{
+		UNIT_LOG_OBJ(Owner, ELogType::StatusDebug, TEXT("NotifySocketSend: Packet dump:"));
+
+		UNIT_LOG_BEGIN(Owner, ELogType::StatusDebug | ELogType::StyleMonospace);
+		NUTDebug::LogHexDump((const uint8*)Data, Count, true, true);
+		UNIT_LOG_END();
+	}
+
+#if !UE_BUILD_SHIPPING
+	LowLevelSendDel.ExecuteIfBound(Data, Count, bBlockSend);
+#endif
+}
+
+void UMinimalClient::NotifyReceivedRawPacket(void* Data, int32 Count, bool& bBlockReceive)
+{
+#if !UE_BUILD_SHIPPING
+	GActiveReceiveUnitConnection = UnitConn;
+
+	ReceivedRawPacketDel.ExecuteIfBound(Data, Count);
+
+	if (Owner != nullptr)
+	{
+		if (!!(MinClientFlags & EMinClientFlags::DumpReceivedRaw))
+		{
+			UNIT_LOG_OBJ(Owner, ELogType::StatusDebug, TEXT("NotifyReceivedRawPacket: Packet dump:"));
+
+			UNIT_LOG_BEGIN(Owner, ELogType::StatusDebug | ELogType::StyleMonospace);
+			NUTDebug::LogHexDump((const uint8*)Data, Count, true, true);
+			UNIT_LOG_END();
+		}
+	}
+
+
+	// The rest of the original ReceivedRawPacket function call is blocked, so temporarily disable the delegate,
+	// and re-trigger it here, so that we correctly encapsulate its call with GActiveReceiveUnitConnection
+	FOnReceivedRawPacket TempDel = UnitConn->ReceivedRawPacketDel;
+
+	UnitConn->ReceivedRawPacketDel.Unbind();
+
+	UnitConn->UNetConnection::ReceivedRawPacket(Data, Count);
+
+	UnitConn->ReceivedRawPacketDel = TempDel;
+
+
+	GActiveReceiveUnitConnection = nullptr;
+
+	// Block the original function call - replaced with the above
+	bBlockReceive = true;
+#endif
+}
+
+void UMinimalClient::NotifyReceiveRPC(AActor* Actor, UFunction* Function, void* Parameters, bool& bBlockRPC)
 {
-	bool bAllowRPC = !!(MinClientFlags & EMinClientFlags::SendRPCs);
+	UNIT_EVENT_BEGIN(Owner);
 
+	// If specified, block RPC's by default - the delegate below has a chance to override this
+	if (!(MinClientFlags & EMinClientFlags::AcceptRPCs))
+	{
+		bBlockRPC = true;
+	}
+
+
+	ReceiveRPCDel.ExecuteIfBound(Actor, Function, Parameters, bBlockRPC);
+
+	FString FuncName = Function->GetName();
+
+	if (bBlockRPC && AllowedClientRPCs.Contains(FuncName))
+	{
+		bBlockRPC = false;
+	}
+
+	if (bBlockRPC)
+	{
+		FString FuncParms = NUTUtilRefl::FunctionParmsToString(Function, Parameters);
+
+		UNIT_LOG_OBJ(Owner,, TEXT("Blocking receive RPC '%s' for actor '%s'"), *FuncName, *Actor->GetFullName());
+
+		if (FuncParms.Len() > 0)
+		{
+			UNIT_LOG_OBJ(Owner,, TEXT("     '%s' parameters: %s"), *FuncName, *FuncParms);
+		}
+	}
+
+
+	if (!!(MinClientFlags & EMinClientFlags::DumpReceivedRPC) && !bBlockRPC)
+	{
+		FString FuncParms = NUTUtilRefl::FunctionParmsToString(Function, Parameters);
+
+		UNIT_LOG_OBJ(Owner, ELogType::StatusDebug, TEXT("Received RPC '%s' for actor '%s'"), *FuncName, *Actor->GetFullName());
+
+		if (FuncParms.Len() > 0)
+		{
+			UNIT_LOG_OBJ(Owner,, TEXT("     '%s' parameters: %s"), *FuncName, *FuncParms);
+		}
+	}
+
+	UNIT_EVENT_END;
+}
+
+void UMinimalClient::NotifySendRPC(AActor* Actor, UFunction* Function, void* Parameters, FOutParmRec* OutParms, FFrame* Stack,
+									UObject* SubObject, bool& bBlockSendRPC)
+{
+	bBlockSendRPC = !(MinClientFlags & EMinClientFlags::SendRPCs);
+
+#if !UE_BUILD_SHIPPING
 	// Pass on to the delegate, and give it an opportunity to override whether the RPC is sent
-	SendRPCDel.ExecuteIfBound(bAllowRPC, Actor, Function, Parameters, OutParms, Stack, SubObject);
+	SendRPCDel.ExecuteIfBound(Actor, Function, Parameters, OutParms, Stack, SubObject, bBlockSendRPC);
+#endif
 
-	if (bAllowRPC)
+	if (!bBlockSendRPC)
 	{
 		if (!!(MinClientFlags & EMinClientFlags::DumpSendRPC))
 		{
@@ -371,8 +699,6 @@ bool UMinimalClient::NotifySendRPC(AActor* Actor, UFunction* Function, void* Par
 						*Actor->GetFullName(), (SubObject != NULL ? *SubObject->GetFullName() : TEXT("nullptr")));
 		}
 	}
-
-	return bAllowRPC;
 }
 
 void UMinimalClient::InternalNotifyNetworkFailure(UWorld* InWorld, UNetDriver* InNetDriver, ENetworkFailure::Type FailureType,
@@ -404,3 +730,24 @@ bool UMinimalClient::NotifyAcceptingChannel(UChannel* Channel)
 
 	return bAccepted;
 }
+
+// @todo #JohnB: Remove dependency on ClientUnitTest eventually
+UMinimalClient* UMinimalClient::GetMinClientFromConn(UNetConnection* InConn)
+{
+	UMinimalClient* ReturnVal = nullptr;
+
+	for (UUnitTest* CurUnitTest : GUnitTestManager->ActiveUnitTests)
+	{
+		UClientUnitTest* CurClientUnitTest = Cast<UClientUnitTest>(CurUnitTest);
+		UMinimalClient* MinClient = (CurClientUnitTest != nullptr ? CurClientUnitTest->MinClient : nullptr);
+		UNetConnection* CurUnitConn = (MinClient != nullptr ? MinClient->GetConn() : nullptr);
+
+		if (CurUnitConn != nullptr && CurUnitConn == InConn)
+		{
+			ReturnVal = MinClient;
+			break;
+		}
+	}
+
+	return ReturnVal;
+}
diff --git a/Engine/Plugins/NetcodeUnitTest/NetcodeUnitTest/Source/NetcodeUnitTest/Private/NUTActor.cpp b/Engine/Plugins/NetcodeUnitTest/NetcodeUnitTest/Source/NetcodeUnitTest/Private/NUTActor.cpp
--- a/Engine/Plugins/NetcodeUnitTest/NetcodeUnitTest/Source/NetcodeUnitTest/Private/NUTActor.cpp
+++ b/Engine/Plugins/NetcodeUnitTest/NetcodeUnitTest/Source/NetcodeUnitTest/Private/NUTActor.cpp
@@ -101,7 +101,7 @@ bool ANUTActor::NotifyControlMessage(UNetConnection* Connection, uint8 MessageTy
 			UpdateOwner();
 		}
 
-		uint8 CmdType = 0;
+		ENUTControlCommand CmdType;
 		FString Command;
 		FNetControlMessage<NMT_NUTControl>::Receive(Bunch, CmdType, Command);
 
@@ -124,7 +124,7 @@ bool ANUTActor::NotifyControlMessage(UNetConnection* Connection, uint8 MessageTy
 
 			if (bSendResult)
 			{
-				uint8 ReturnCmdType = (bCmdSuccess ? ENUTControlCommand::CommandResult_Success :
+				ENUTControlCommand ReturnCmdType = (bCmdSuccess ? ENUTControlCommand::CommandResult_Success :
 										ENUTControlCommand::CommandResult_Failed);
 
 				FNetControlMessage<NMT_NUTControl>::Send(Connection, ReturnCmdType, CmdResult);
@@ -148,7 +148,7 @@ bool ANUTActor::NotifyControlMessage(UNetConnection* Connection, uint8 MessageTy
 		// Ping request
 		else if (CmdType == ENUTControlCommand::Ping)
 		{
-			uint8 TempCmdType = ENUTControlCommand::Pong;
+			ENUTControlCommand TempCmdType = ENUTControlCommand::Pong;
 			FString Dud;
 
 			FNetControlMessage<NMT_NUTControl>::Send(Connection, TempCmdType, Dud);
@@ -298,7 +298,7 @@ void ANUTActor::NotifyPostLoadMap(UWorld* LoadedWorld)
 {
 	if (VerifyEventWatcher())
 	{
-		uint8 CmdType = ENUTControlCommand::NotifyEvent;
+		ENUTControlCommand CmdType = ENUTControlCommand::NotifyEvent;
 		FString Command = TEXT("NotifySeamlessTravelEnd");
 
 		FNetControlMessage<NMT_NUTControl>::Send(EventWatcher, CmdType, Command);
diff --git a/Engine/Plugins/NetcodeUnitTest/NetcodeUnitTest/Source/NetcodeUnitTest/Private/NUTUtilDebug.cpp b/Engine/Plugins/NetcodeUnitTest/NetcodeUnitTest/Source/NetcodeUnitTest/Private/NUTUtilDebug.cpp
--- a/Engine/Plugins/NetcodeUnitTest/NetcodeUnitTest/Source/NetcodeUnitTest/Private/NUTUtilDebug.cpp
+++ b/Engine/Plugins/NetcodeUnitTest/NetcodeUnitTest/Source/NetcodeUnitTest/Private/NUTUtilDebug.cpp
@@ -35,16 +35,16 @@ void FScopedLog::InternalConstruct(const TArray<FString>& InLogCategories, UClie
 	bRemoteLogging = bInRemoteLogging;
 
 
+	UMinimalClient* MinClient = (UnitTest != nullptr ? UnitTest->MinClient : nullptr);
+	UNetConnection* UnitConn = (MinClient != nullptr ? MinClient->GetConn() : nullptr);
+
 	// If you want to do remote logging, you MUST specify the client unit test doing the logging
 	if (bRemoteLogging)
 	{
 		UNIT_ASSERT(UnitTest != nullptr);
-		UNIT_ASSERT(UnitTest->UnitConn != nullptr);
+		UNIT_ASSERT(UnitConn != nullptr);
 	}
 
-
-	UNetConnection* UnitConn = (UnitTest != nullptr ? UnitTest->UnitConn : nullptr);
-
 	// Flush all current packets, so the log messages only relate to scoped code
 	if (UnitConn != nullptr)
 	{
@@ -53,14 +53,14 @@ void FScopedLog::InternalConstruct(const TArray<FString>& InLogCategories, UClie
 
 
 	// If specified, enable logs remotely
-	if (bRemoteLogging && UnitTest != nullptr)
+	if (bRemoteLogging && MinClient != nullptr)
 	{
-		FOutBunch* ControlChanBunch = NUTNet::CreateChannelBunch(UnitTest->ControlBunchSequence, UnitConn, CHTYPE_Control, 0);
+		FOutBunch* ControlChanBunch = MinClient->CreateChannelBunch(CHTYPE_Control, 0);
 
 		if (ControlChanBunch != nullptr)
 		{
 			uint8 ControlMsg = NMT_NUTControl;
-			uint8 ControlCmd = ENUTControlCommand::Command_NoResult;
+			ENUTControlCommand ControlCmd = ENUTControlCommand::Command_NoResult;
 			FString Cmd = TEXT("");
 
 			for (auto CurCategory : LogCategories)
@@ -72,7 +72,7 @@ void FScopedLog::InternalConstruct(const TArray<FString>& InLogCategories, UClie
 				*ControlChanBunch << Cmd;
 			}
 
-			NUTNet::SendControlBunch(UnitConn, *ControlChanBunch);
+			MinClient->SendControlBunch(ControlChanBunch);
 
 
 			// Need to flush again now to get the above parsed on server first
@@ -83,9 +83,9 @@ void FScopedLog::InternalConstruct(const TArray<FString>& InLogCategories, UClie
 
 	// Now enable local logging
 	FString Cmd = TEXT("");
-	UWorld* UnitWorld = (UnitTest != nullptr && UnitTest->MinClient != nullptr) ? UnitTest->MinClient->GetUnitWorld() : nullptr;
+	UWorld* UnitWorld = (MinClient != nullptr ? MinClient->GetUnitWorld() : nullptr);
 
-	for (auto CurCategory : LogCategories)
+	for (FString CurCategory : LogCategories)
 	{
 		Cmd = TEXT("Log ") + CurCategory + TEXT(" All");
 
@@ -95,10 +95,11 @@ void FScopedLog::InternalConstruct(const TArray<FString>& InLogCategories, UClie
 
 FScopedLog::~FScopedLog()
 {
-	UNetConnection* UnitConn = (UnitTest != NULL ? UnitTest->UnitConn : NULL);
+	UMinimalClient* MinClient = (UnitTest != nullptr ? UnitTest->MinClient : nullptr);
+	UNetConnection* UnitConn = (MinClient != nullptr ? MinClient->GetConn() : nullptr);
 
 	// Flush all built-up packets
-	if (UnitConn != NULL)
+	if (UnitConn != nullptr)
 	{
 		UnitConn->FlushNet();
 	}
@@ -106,7 +107,7 @@ FScopedLog::~FScopedLog()
 
 	// Reset local logging
 	FString Cmd = TEXT("");
-	UWorld* UnitWorld = (UnitTest != nullptr && UnitTest->MinClient != nullptr) ? UnitTest->MinClient->GetUnitWorld() : nullptr;
+	UWorld* UnitWorld = (MinClient != nullptr ? MinClient->GetUnitWorld() : nullptr);
 
 	for (int32 i=LogCategories.Num()-1; i>=0; i--)
 	{
@@ -117,14 +118,14 @@ FScopedLog::~FScopedLog()
 
 
 	// Reset remote logging (and flush immediately)
-	if (bRemoteLogging && UnitTest != nullptr)
+	if (bRemoteLogging && MinClient != nullptr)
 	{
-		FOutBunch* ControlChanBunch = NUTNet::CreateChannelBunch(UnitTest->ControlBunchSequence, UnitConn, CHTYPE_Control, 0);
+		FOutBunch* ControlChanBunch = MinClient->CreateChannelBunch(CHTYPE_Control, 0);
 
 		if (ControlChanBunch != nullptr)
 		{
 			uint8 ControlMsg = NMT_NUTControl;
-			uint8 ControlCmd = ENUTControlCommand::Command_NoResult;
+			ENUTControlCommand ControlCmd = ENUTControlCommand::Command_NoResult;
 
 			for (int32 i=LogCategories.Num()-1; i>=0; i--)
 			{
@@ -135,14 +136,15 @@ FScopedLog::~FScopedLog()
 				*ControlChanBunch << Cmd;
 			}
 
-			NUTNet::SendControlBunch(UnitConn, *ControlChanBunch);
+			MinClient->SendControlBunch(ControlChanBunch);
 
 			UnitConn->FlushNet();
 		}
 	}
 }
 
-
+// @todo #JohnB: Reimplement eventually - see header file, needs to be merged with a similar class
+#if 0
 #if !UE_BUILD_SHIPPING
 /**
  * FProcessEventHookBase
@@ -179,7 +181,11 @@ bool FProcessEventHookBase::InternalProcessEventHook(AActor* Actor, UFunction* F
 
 	return bReturnVal;
 }
+#endif
+#endif
 
+// @todo #JohnB: Reimplement when needed. See header file.
+#if 0
 void FScopedProcessEventLog::ProcessEventHook(AActor* Actor, UFunction* Function, void* Parameters)
 {
 	UE_LOG(LogUnitTest, Log, TEXT("FScopedProcessEventLog: Actor: %s, Function: %s"),
diff --git a/Engine/Plugins/NetcodeUnitTest/NetcodeUnitTest/Source/NetcodeUnitTest/Private/NUTUtilReflection.cpp b/Engine/Plugins/NetcodeUnitTest/NetcodeUnitTest/Source/NetcodeUnitTest/Private/NUTUtilReflection.cpp
--- a/Engine/Plugins/NetcodeUnitTest/NetcodeUnitTest/Source/NetcodeUnitTest/Private/NUTUtilReflection.cpp
+++ b/Engine/Plugins/NetcodeUnitTest/NetcodeUnitTest/Source/NetcodeUnitTest/Private/NUTUtilReflection.cpp
@@ -345,26 +345,26 @@ FVMReflection& FVMReflection::operator [](const ANSICHAR* InFieldType)
 	UArrayProperty* ArrayProp = Cast<UArrayProperty>(FieldProp);
 	UStruct* StructField = Cast<UStruct>(FieldInstance);
 
-	if (!bIsError && FieldInstance != NULL && FieldAddress != NULL && !bVerifiedFieldType)
+	if (!bIsError && FieldInstance != nullptr && FieldAddress != nullptr && !bVerifiedFieldType)
 	{
-		UField* ActualFieldType = NULL;
+		UField* ActualFieldType = nullptr;
 		FString ActualFieldTypeStr = TEXT("");
 		const TCHAR* CheckType = TEXT("");
 
 		// Static arrays
-		if (FieldProp != NULL && FieldProp->ArrayDim > 1)
+		if (FieldProp != nullptr && FieldProp->ArrayDim > 1)
 		{
 			ActualFieldType = FieldProp;
 			CheckType = TEXT("array");
 		}
 		// Dynamic arrays
-		else if (ArrayProp != NULL)
+		else if (ArrayProp != nullptr)
 		{
 			ActualFieldType = ArrayProp->Inner;
 			CheckType = TEXT("array");
 		}
 		// Structs
-		else if (StructField != NULL)
+		else if (StructField != nullptr)
 		{
 			ActualFieldType = StructField;
 			CheckType = TEXT("struct");
@@ -377,14 +377,14 @@ FVMReflection& FVMReflection::operator [](const ANSICHAR* InFieldType)
 		// Handle special case of struct-arrays
 		UStructProperty* StructProp = Cast<UStructProperty>(ActualFieldType);
 
-		if (StructProp != NULL)
+		if (StructProp != nullptr)
 		{
 			ActualFieldType = StructProp->Struct;
 			CheckType = TEXT("struct array");
 		}
 
 
-		if (ActualFieldType != NULL)
+		if (ActualFieldType != nullptr)
 		{
 			bool bTypeValid = false;
 
@@ -488,7 +488,7 @@ FVMReflection& FVMReflection::operator [](const ANSICHAR* InFieldType)
 			{
 				UStruct* StructRef = Cast<UStruct>(ActualFieldType);
 
-				if (StructRef != NULL)
+				if (StructRef != nullptr)
 				{
 					FString ClassName = ExpectedFieldType.Mid(1);
 
diff --git a/Engine/Plugins/NetcodeUnitTest/NetcodeUnitTest/Source/NetcodeUnitTest/Private/Net/NUTUtilNet.cpp b/Engine/Plugins/NetcodeUnitTest/NetcodeUnitTest/Source/NetcodeUnitTest/Private/Net/NUTUtilNet.cpp
--- a/Engine/Plugins/NetcodeUnitTest/NetcodeUnitTest/Source/NetcodeUnitTest/Private/Net/NUTUtilNet.cpp
+++ b/Engine/Plugins/NetcodeUnitTest/NetcodeUnitTest/Source/NetcodeUnitTest/Private/Net/NUTUtilNet.cpp
@@ -4,25 +4,22 @@
 #include "UObject/CoreOnline.h"
 #include "GameFramework/OnlineReplStructs.h"
 #include "Engine/Engine.h"
-#include "Engine/GameEngine.h"
 #include "PacketHandlers/StatelessConnectHandlerComponent.h"
 #include "GameFramework/PlayerController.h"
-#include "Engine/NetDriver.h"
 #include "Net/DataBunch.h"
 #include "Engine/LocalPlayer.h"
 #include "EngineUtils.h"
-#include "IpNetDriver.h"
 #include "Net/DataChannel.h"
-#include "OnlineBeaconClient.h"
 
 #include "NUTUtil.h"
+#include "NUTUtilReflection.h"
 #include "UnitTest.h"
 #include "ClientUnitTest.h"
 #include "MinimalClient.h"
 #include "UnitTestEnvironment.h"
 #include "Net/UnitTestPackageMap.h"
 #include "Net/UnitTestChannel.h"
-#include "Net/UnitTestNetDriver.h"
+
 
 // Forward declarations
 class FWorldTickHook;
@@ -39,164 +36,91 @@ TArray<FWorldTickHook*> ActiveTickHooks;
 
 
 /**
- * FSocketHook
+ * FProcessNetEventHook
  */
 
-bool FSocketHook::Close()
-{
-	return HookedSocket->Close();
-}
-
-bool FSocketHook::Bind(const FInternetAddr& Addr)
-{
-	return HookedSocket->Bind(Addr);
-}
-
-bool FSocketHook::Connect(const FInternetAddr& Addr)
-{
-	return HookedSocket->Connect(Addr);
-}
-
-bool FSocketHook::Listen(int32 MaxBacklog)
-{
-	return HookedSocket->Listen(MaxBacklog);
-}
-
-bool FSocketHook::HasPendingConnection(bool& bHasPendingConnection)
+void FProcessEventHook::AddRPCHook(UWorld* InWorld, FOnProcessNetEvent InHook)
 {
-	return HookedSocket->HasPendingConnection(bHasPendingConnection);
-}
-
-bool FSocketHook::HasPendingData(uint32& PendingDataSize)
-{
-	return HookedSocket->HasPendingData(PendingDataSize);
-}
-
-class FSocket* FSocketHook::Accept(const FString& InSocketDescription)
-{
-	return HookedSocket->Accept(InSocketDescription);
-}
-
-class FSocket* FSocketHook::Accept(FInternetAddr& OutAddr, const FString& InSocketDescription)
-{
-	return HookedSocket->Accept(OutAddr, InSocketDescription);
-}
-
-bool FSocketHook::SendTo(const uint8* Data, int32 Count, int32& BytesSent, const FInternetAddr& Destination)
-{
-	bool bReturnVal = false;
-	bool bBlockSend = false;
-
-	if (MinClient != nullptr)
+#if !UE_BUILD_SHIPPING
+	if (NetEventHooks.Num() == 0 && EventHooks.Num() == 0)
 	{
-		MinClient->SocketSendDel.ExecuteIfBound((void*)Data, Count, bBlockSend);
+		AActor::ProcessEventDelegate.BindRaw(this, &FProcessEventHook::HandleProcessEvent);
 	}
 
-	SendToDel.ExecuteIfBound((void*)Data, Count, bBlockSend);
-
-	if (!bBlockSend)
-	{
-		bReturnVal = HookedSocket->SendTo(Data, Count, BytesSent, Destination);
-	}
-
-	bLastSendToBlocked = bBlockSend;
-
-	return bReturnVal;
-}
-
-bool FSocketHook::Send(const uint8* Data, int32 Count, int32& BytesSent)
-{
-	return HookedSocket->Send(Data, Count, BytesSent);
-}
-
-bool FSocketHook::RecvFrom(uint8* Data, int32 BufferSize, int32& BytesRead, FInternetAddr& Source, ESocketReceiveFlags::Type Flags)
-{
-	return HookedSocket->RecvFrom(Data, BufferSize, BytesRead, Source, Flags);
-}
-
-bool FSocketHook::Recv(uint8* Data, int32 BufferSize, int32& BytesRead, ESocketReceiveFlags::Type Flags)
-{
-	return HookedSocket->Recv(Data, BufferSize, BytesRead, Flags);
-}
-
-bool FSocketHook::Wait(ESocketWaitConditions::Type Condition, FTimespan WaitTime)
-{
-	return HookedSocket->Wait(Condition, WaitTime);
-}
-
-ESocketConnectionState FSocketHook::GetConnectionState()
-{
-	return HookedSocket->GetConnectionState();
-}
-
-void FSocketHook::GetAddress(FInternetAddr& OutAddr)
-{
-	HookedSocket->GetAddress(OutAddr);
-}
-
-bool FSocketHook::GetPeerAddress(FInternetAddr& OutAddr)
-{
-	return HookedSocket->GetPeerAddress(OutAddr);
+	NetEventHooks.Add(InWorld, InHook);
+#else
+	check(false);
+#endif
 }
 
-bool FSocketHook::SetNonBlocking(bool bIsNonBlocking)
+void FProcessEventHook::RemoveRPCHook(UWorld* InWorld)
 {
-	return HookedSocket->SetNonBlocking(bIsNonBlocking);
-}
+#if !UE_BUILD_SHIPPING
+	NetEventHooks.Remove(InWorld);
 
-bool FSocketHook::SetBroadcast(bool bAllowBroadcast)
-{
-	return HookedSocket->SetBroadcast(bAllowBroadcast);
+	if (NetEventHooks.Num() == 0 && EventHooks.Num() == 0)
+	{
+		AActor::ProcessEventDelegate.Unbind();
+	}
+#else
+	check(false);
+#endif
 }
 
-bool FSocketHook::JoinMulticastGroup(const FInternetAddr& GroupAddress)
+void FProcessEventHook::AddEventHook(UWorld* InWorld, FOnProcessNetEvent InHook)
 {
-	return HookedSocket->JoinMulticastGroup(GroupAddress);
-}
+#if !UE_BUILD_SHIPPING
+	if (NetEventHooks.Num() == 0 && EventHooks.Num() == 0)
+	{
+		AActor::ProcessEventDelegate.BindRaw(this, &FProcessEventHook::HandleProcessEvent);
+	}
 
-bool FSocketHook::LeaveMulticastGroup(const FInternetAddr& GroupAddress)
-{
-	return HookedSocket->LeaveMulticastGroup(GroupAddress);
+	EventHooks.Add(InWorld, InHook);
+#else
+	check(false);
+#endif
 }
 
-bool FSocketHook::SetMulticastLoopback(bool bLoopback)
+void FProcessEventHook::RemoveEventHook(UWorld* InWorld)
 {
-	return HookedSocket->SetMulticastLoopback(bLoopback);
-}
+#if !UE_BUILD_SHIPPING
+	EventHooks.Remove(InWorld);
 
-bool FSocketHook::SetMulticastTtl(uint8 TimeToLive)
-{
-	return HookedSocket->SetMulticastTtl(TimeToLive);
+	if (NetEventHooks.Num() == 0 && EventHooks.Num() == 0)
+	{
+		AActor::ProcessEventDelegate.Unbind();
+	}
+#else
+	check(false);
+#endif
 }
 
-bool FSocketHook::SetReuseAddr(bool bAllowReuse)
+bool FProcessEventHook::HandleProcessEvent(AActor* Actor, UFunction* Function, void* Parameters)
 {
-	return HookedSocket->SetReuseAddr(bAllowReuse);
-}
+	bool bBlockEvent = false;
 
-bool FSocketHook::SetLinger(bool bShouldLinger, int32 Timeout)
-{
-	return HookedSocket->SetLinger(bShouldLinger, Timeout);
-}
+	if (Actor != nullptr && Function != nullptr)
+	{
+		bool bNetClientRPC = !!(Function->FunctionFlags & FUNC_Net) && !!(Function->FunctionFlags & FUNC_NetClient);
+		bool bValidEvent = (bNetClientRPC && NetEventHooks.Num() > 0) || EventHooks.Num() > 0;
 
-bool FSocketHook::SetRecvErr(bool bUseErrorQueue)
-{
-	return HookedSocket->SetRecvErr(bUseErrorQueue);
-}
+		if (bValidEvent)
+		{
+			UWorld* CurWorld = Actor->GetWorld();
 
-bool FSocketHook::SetSendBufferSize(int32 Size, int32& NewSize)
-{
-	return HookedSocket->SetSendBufferSize(Size, NewSize);
-}
+			if (CurWorld != nullptr)
+			{
+				FOnProcessNetEvent* Hook = (bNetClientRPC ? NetEventHooks.Find(CurWorld) : EventHooks.Find(CurWorld));
 
-bool FSocketHook::SetReceiveBufferSize(int32 Size, int32& NewSize)
-{
-	return HookedSocket->SetReceiveBufferSize(Size, NewSize);
-}
+				if (Hook != nullptr)
+				{
+					Hook->Execute(Actor, Function, Parameters, bBlockEvent);
+				}
+			}
+		}
+	}
 
-int32 FSocketHook::GetPortNo()
-{
-	return HookedSocket->GetPortNo();
+	return bBlockEvent;
 }
 
 
@@ -275,9 +199,11 @@ FScopedNetObjectReplace::FScopedNetObjectReplace(UClientUnitTest* InUnitTest, UO
 	: UnitTest(InUnitTest)
 	, ObjToReplace(InObjToReplace)
 {
-	if (UnitTest != nullptr && UnitTest->UnitConn != nullptr)
+	UNetConnection* UnitConn = (UnitTest != nullptr && UnitTest->MinClient != nullptr ? UnitTest->MinClient->GetConn() : nullptr);
+
+	if (UnitConn != nullptr)
 	{
-		UUnitTestPackageMap* PackageMap = Cast<UUnitTestPackageMap>(UnitTest->UnitConn->PackageMap);
+		UUnitTestPackageMap* PackageMap = Cast<UUnitTestPackageMap>(UnitConn->PackageMap);
 
 		if (PackageMap != nullptr)
 		{
@@ -298,9 +224,11 @@ FScopedNetObjectReplace::FScopedNetObjectReplace(UClientUnitTest* InUnitTest, UO
 
 FScopedNetObjectReplace::~FScopedNetObjectReplace()
 {
-	if (UnitTest != nullptr && UnitTest->UnitConn != nullptr)
+	UNetConnection* UnitConn = (UnitTest != nullptr && UnitTest->MinClient != nullptr ? UnitTest->MinClient->GetConn() : nullptr);
+
+	if (UnitConn != nullptr)
 	{
-		UUnitTestPackageMap* PackageMap = Cast<UUnitTestPackageMap>(UnitTest->UnitConn->PackageMap);
+		UUnitTestPackageMap* PackageMap = Cast<UUnitTestPackageMap>(UnitConn->PackageMap);
 
 		if (PackageMap != nullptr)
 		{
@@ -324,232 +252,13 @@ FScopedNetObjectReplace::~FScopedNetObjectReplace()
  * NUTNet
  */
 
-int32 NUTNet::GetFreeChannelIndex(UNetConnection* InConnection, int32 InChIndex/*=INDEX_NONE*/)
-{
-	int32 ReturnVal = INDEX_NONE;
-
-	for (int i=(InChIndex == INDEX_NONE ? 0 : InChIndex); i<ARRAY_COUNT(InConnection->Channels); i++)
-	{
-		if (InConnection->Channels[i] == NULL)
-		{
-			ReturnVal = i;
-			break;
-		}
-	}
-
-	return ReturnVal;
-}
-
-UUnitTestChannel* NUTNet::CreateUnitTestChannel(UNetConnection* InConnection, EChannelType InType, int32 InChIndex/*=INDEX_NONE*/,
-												bool bInVerifyOpen/*=false*/)
-{
-	UUnitTestChannel* ReturnVal = NULL;
-
-	if (InChIndex == INDEX_NONE)
-	{
-		InChIndex = GetFreeChannelIndex(InConnection);
-	}
-
-	if (InConnection != NULL && InChIndex >= 0 && InChIndex < ARRAY_COUNT(InConnection->Channels) &&
-		InConnection->Channels[InChIndex] == NULL)
-	{
-		ReturnVal = NewObject<UUnitTestChannel>();
-		ReturnVal->ChType = InType;
-
-		// Always treat as not opened locally
-		ReturnVal->Init(InConnection, InChIndex, false);
-		ReturnVal->bVerifyOpen = bInVerifyOpen;
-
-		InConnection->Channels[InChIndex] = ReturnVal;
-		InConnection->OpenChannels.Add(ReturnVal);
-
-		UE_LOG(LogUnitTest, Log, TEXT("Created unit test channel of type '%i' on ChIndex '%i'"), (int)InType, InChIndex);
-	}
-	else if (InConnection == NULL)
-	{
-		UE_LOG(LogUnitTest, Log, TEXT("Failed to create unit test channel: InConnection == NULL"));
-	}
-	else if (InChIndex < 0 || InChIndex >= ARRAY_COUNT(InConnection->Channels))
-	{
-		UE_LOG(LogUnitTest, Log, TEXT("Failed to create unit test channel: InChIndex <= 0 || >= channel count (InchIndex: %i)"),
-				InChIndex);
-	}
-	else if (InConnection->Channels[InChIndex] != NULL)
-	{
-		UE_LOG(LogUnitTest, Log, TEXT("Failed to create unit test channel: Channel index already set"));
-	}
-
-	return ReturnVal;
-}
-
-FOutBunch* NUTNet::CreateChannelBunch(int32& BunchSequence, UNetConnection* InConnection, EChannelType InChType,
-											int32 InChIndex/*=INDEX_NONE*/, bool bGetNextFreeChan/*=false*/)
-{
-	FOutBunch* ReturnVal = nullptr;
-	UChannel* ControlChan = InConnection != nullptr ? InConnection->Channels[0] : nullptr;
-
-	if (InChIndex == INDEX_NONE || bGetNextFreeChan)
-	{
-		InChIndex = GetFreeChannelIndex(InConnection, InChIndex);
-	}
-
-	if (ControlChan != nullptr && ControlChan->IsNetReady(false))
-	{
-		ReturnVal = new FOutBunch(ControlChan, false);
-
-		// Fix for uninitialized members (just one is causing a problem, but may as well do them all)
-		ReturnVal->Next = nullptr;
-		ReturnVal->Time = 0.0;
-		ReturnVal->ReceivedAck = false; // This one was the problem - was triggering an assert
-		ReturnVal->ChSequence = 0;
-		ReturnVal->PacketId = 0;
-		ReturnVal->bDormant = false;
-
-		// If the channel already exists, the BunchSequence needs to be overridden, and the channel sequence incremented
-		if (InConnection->Channels[InChIndex] != nullptr)
-		{
-			BunchSequence = ++InConnection->OutReliable[InChIndex];
-		}
-		else
-		{
-			BunchSequence++;
-		}
-
-		ReturnVal->Channel = nullptr;
-		ReturnVal->ChIndex = InChIndex;
-		ReturnVal->ChType = InChType;
-
-
-		// NOTE: Might not cover all bOpen or 'channel already open' cases
-		if (InConnection->Channels[InChIndex] == nullptr)
-		{
-			ReturnVal->bOpen = 1;
-		}
-		else if (InConnection->Channels[InChIndex]->OpenPacketId.First == INDEX_NONE)
-		{
-			ReturnVal->bOpen = 1;
-
-			InConnection->Channels[InChIndex]->OpenPacketId.First = BunchSequence;
-			InConnection->Channels[InChIndex]->OpenPacketId.Last = BunchSequence;
-		}
-
-
-		ReturnVal->bReliable = 1;
-		ReturnVal->ChSequence = BunchSequence;
-	}
-
-	return ReturnVal;
-}
-
-void NUTNet::SendControlBunch(UNetConnection* InConnection, FOutBunch& ControlChanBunch)
-{
-	if (InConnection != NULL && InConnection->Channels[0] != NULL)
-	{
-		// Since it's the unit test control channel, sending the packet abnormally, append to OutRec manually
-		if (ControlChanBunch.bReliable)
-		{
-			UChannel* ControlChan = InConnection->Channels[0];
-
-			for (FOutBunch* CurOut=ControlChan->OutRec; CurOut!=NULL; CurOut=CurOut->Next)
-			{
-				if (CurOut->Next == NULL)
-				{
-					CurOut->Next = &ControlChanBunch;
-					ControlChan->NumOutRec++;
-
-					break;
-				}
-			}
-		}
-
-		InConnection->SendRawBunch(ControlChanBunch, true);
-	}
-}
-
-UUnitTestNetDriver* NUTNet::CreateUnitTestNetDriver(UWorld* InWorld)
-{
-	UUnitTestNetDriver* ReturnVal = NULL;
-	UGameEngine* GameEngine = Cast<UGameEngine>(GEngine);
-
-	if (GameEngine != NULL)
-	{
-		static int UnitTestNetDriverCount = 0;
-
-		// Setup a new driver name entry
-		bool bFoundDef = false;
-		FName UnitDefName = TEXT("UnitTestNetDriver");
-
-		for (int i=0; i<GameEngine->NetDriverDefinitions.Num(); i++)
-		{
-			if (GameEngine->NetDriverDefinitions[i].DefName == UnitDefName)
-			{
-				bFoundDef = true;
-				break;
-			}
-		}
-
-		if (!bFoundDef)
-		{
-			FNetDriverDefinition NewDriverEntry;
-
-			NewDriverEntry.DefName = UnitDefName;
-			NewDriverEntry.DriverClassName = *UUnitTestNetDriver::StaticClass()->GetPathName();
-			NewDriverEntry.DriverClassNameFallback = *UIpNetDriver::StaticClass()->GetPathName();
-
-			GameEngine->NetDriverDefinitions.Add(NewDriverEntry);
-		}
-
-
-		FName NewDriverName = *FString::Printf(TEXT("UnitTestNetDriver_%i"), UnitTestNetDriverCount++);
-
-		// Now create a reference to the driver
-		if (GameEngine->CreateNamedNetDriver(InWorld, NewDriverName, UnitDefName))
-		{
-			ReturnVal = Cast<UUnitTestNetDriver>(GameEngine->FindNamedNetDriver(InWorld, NewDriverName));
-		}
-
-
-		if (ReturnVal != NULL)
-		{
-			ReturnVal->SetWorld(InWorld);
-			InWorld->SetNetDriver(ReturnVal);
-
-
-			FLevelCollection* Collection = (FLevelCollection*)InWorld->GetActiveLevelCollection();
-
-			// Hack-set the net driver in the worlds level collection
-			if (Collection != nullptr)
-			{
-				Collection->SetNetDriver(ReturnVal);
-			}
-			else
-			{
-				UE_LOG(LogUnitTest, Warning,
-						TEXT("CreateUnitTestNetDriver: No LevelCollection found for created world, may block replication."));
-			}
-
-			UE_LOG(LogUnitTest, Log, TEXT("CreateUnitTestNetDriver: Created named net driver: %s, NetDriverName: %s, for World: %s"),
-					*ReturnVal->GetFullName(), *ReturnVal->NetDriverName.ToString(), *InWorld->GetFullName());
-		}
-		else
-		{
-			UE_LOG(LogUnitTest, Log, TEXT("CreateUnitTestNetDriver: CreateNamedNetDriver failed"));
-		}
-	}
-	else
-	{
-		UE_LOG(LogUnitTest, Log, TEXT("CreateUnitTestNetDriver: GameEngine is NULL"));
-	}
-
-	return ReturnVal;
-}
-
-void NUTNet::HandleBeaconReplicate(AOnlineBeaconClient* InBeacon, UNetConnection* InConnection)
+void NUTNet::HandleBeaconReplicate(AActor* InBeacon, UNetConnection* InConnection)
 {
 	// Due to the way the beacon is created in unit tests (replicated, instead of taking over a local beacon client),
 	// the NetDriver and BeaconConnection values have to be hack-set, to enable sending of RPC's
 	InBeacon->SetNetDriverName(InConnection->Driver->NetDriverName);
-	InBeacon->SetNetConnection(InConnection);
+
+	FVMReflection(InBeacon)->*"BeaconConnection" = InConnection;
 }
 
 
diff --git a/Engine/Plugins/NetcodeUnitTest/NetcodeUnitTest/Source/NetcodeUnitTest/Private/Net/UnitTestActorChannel.cpp b/Engine/Plugins/NetcodeUnitTest/NetcodeUnitTest/Source/NetcodeUnitTest/Private/Net/UnitTestActorChannel.cpp
--- a/Engine/Plugins/NetcodeUnitTest/NetcodeUnitTest/Source/NetcodeUnitTest/Private/Net/UnitTestActorChannel.cpp
+++ b/Engine/Plugins/NetcodeUnitTest/NetcodeUnitTest/Source/NetcodeUnitTest/Private/Net/UnitTestActorChannel.cpp
@@ -1,9 +1,17 @@
 // Copyright 1998-2017 Epic Games, Inc. All Rights Reserved.
 
 #include "Net/UnitTestActorChannel.h"
+#include "Engine/Engine.h"
+#include "Engine/LocalPlayer.h"
 #include "Engine/NetConnection.h"
+#include "GameFramework/PlayerController.h"
+
 #include "NetcodeUnitTest.h"
 
+#include "UnitTestManager.h"
+#include "ClientUnitTest.h"
+#include "MinimalClient.h"
+
 
 
 /**
@@ -11,53 +19,86 @@
  */
 UUnitTestActorChannel::UUnitTestActorChannel(const FObjectInitializer& ObjectInitializer)
 	: Super(ObjectInitializer)
+	, MinClient(nullptr)
+{
+#if !UE_BUILD_SHIPPING
+	bBlockChannelFailure = true;
+#endif
+}
+
+void UUnitTestActorChannel::Init(UNetConnection* InConnection, int32 InChIndex, bool InOpenedLocally)
 {
+	MinClient = UMinimalClient::GetMinClientFromConn(InConnection);
+
+	Super::Init(InConnection, InChIndex, InOpenedLocally);
 }
 
 void UUnitTestActorChannel::ReceivedBunch(FInBunch& Bunch)
 {
-	// If Actor is NULL, and we are processing a bunch, then we are most likely initializing the actor channel
-	if (Actor == NULL)
+	// If Actor is nullptr, and we are processing a bunch, then we are most likely initializing the actor channel
+	if (Actor == nullptr)
 	{
 		GIsInitializingActorChan = true;
 	}
 
-
-	// As a hack to block sending an 'NMT_ActorChannelFailed' event that kicks the client,
-	// the control channel is NULL'd when an actor channel is blocked, and this is in charge of recovering from that
-	// NOTE: If you update this, update the below function too
-	UChannel* ControlChan = Connection->Channels[0];
-
 	Super::ReceivedBunch(Bunch);
 
-	if (ControlChan != NULL && Connection->Channels[0] == NULL)
-	{
-		Connection->Channels[0] = ControlChan;
-	}
-
 	GIsInitializingActorChan = false;
 }
 
 void UUnitTestActorChannel::Tick()
 {
-	// If Actor is NULL and we're about to process queued bunches, we may be about to do a (delayed) actor channel initialization
-	if (Actor == NULL && QueuedBunches.Num() > 0)
+	UNetConnection* OldReceiveUnitConn = GActiveReceiveUnitConnection;
+
+	// If Actor is nullptr and we're about to process queued bunches, we may be about to do a (delayed) actor channel initialization
+	if (Actor == nullptr && QueuedBunches.Num() > 0)
 	{
 		GIsInitializingActorChan = true;
 		GActiveReceiveUnitConnection = Connection;
 	}
 
-	// NOTE: If you update this, update the above function too
-	UChannel* ControlChan = Connection->Channels[0];
-
 	Super::Tick();
 
-	if (ControlChan != NULL && Connection->Channels[0] == NULL)
-	{
-		Connection->Channels[0] = ControlChan;
-	}
-
-	GActiveReceiveUnitConnection = NULL;
+	GActiveReceiveUnitConnection = OldReceiveUnitConn;
 	GIsInitializingActorChan = false;
 }
 
+void UUnitTestActorChannel::NotifyActorChannelOpen(AActor* InActor, FInBunch& InBunch)
+{
+	APlayerController* PC = Cast<APlayerController>(InActor);
+
+	// Override 'UNetConnection::HandleClientPlayer', as called through PlayerController::OnActorChannelOpen
+	if (PC != nullptr)
+	{
+		uint8 NetPlayerIndex = 0;
+
+		InBunch << NetPlayerIndex;
+
+		// Does not support child connections at the moment
+		check(NetPlayerIndex == 0);
+
+		if (NetPlayerIndex == 0 && Connection->Driver != nullptr && Connection == Connection->Driver->ServerConnection)
+		{
+			// Implement only essential parts of the original function, as we want to block most of it (triggers level change code)
+			PC->Role = ROLE_AutonomousProxy;
+			PC->NetConnection = Connection;
+
+			Connection->PlayerController = PC;
+			Connection->OwningActor = PC;
+
+			if (PC->Player == nullptr)
+			{
+				PC->Player = NewObject<ULocalPlayer>(GEngine, GEngine->LocalPlayerClass);
+			}
+
+			if (MinClient != nullptr)
+			{
+				MinClient->HandleClientPlayerDel.ExecuteIfBound(PC, Connection);
+			}
+		}
+	}
+	else
+	{
+		Super::NotifyActorChannelOpen(InActor, InBunch);
+	}
+}
diff --git a/Engine/Plugins/NetcodeUnitTest/NetcodeUnitTest/Source/NetcodeUnitTest/Private/Net/UnitTestPackageMap.cpp b/Engine/Plugins/NetcodeUnitTest/NetcodeUnitTest/Source/NetcodeUnitTest/Private/Net/UnitTestPackageMap.cpp
--- a/Engine/Plugins/NetcodeUnitTest/NetcodeUnitTest/Source/NetcodeUnitTest/Private/Net/UnitTestPackageMap.cpp
+++ b/Engine/Plugins/NetcodeUnitTest/NetcodeUnitTest/Source/NetcodeUnitTest/Private/Net/UnitTestPackageMap.cpp
@@ -5,12 +5,12 @@
 #include "GameFramework/Actor.h"
 
 #include "MinimalClient.h"
-#include "Net/UnitTestNetConnection.h"
 
 
 /**
- * Default constructor
+ * UUnitTestPackageMap
  */
+
 UUnitTestPackageMap::UUnitTestPackageMap(const FObjectInitializer& ObjectInitializer)
 	: Super(ObjectInitializer)
 	, bWithinSerializeNewActor(false)
@@ -57,20 +57,19 @@ bool UUnitTestPackageMap::SerializeObject(FArchive& Ar, UClass* InClass, UObject
 		// this is the first place we know the type of a replicated actor (in an actor channel or otherwise), but BEFORE it is spawned
 		else if ((GIsInitializingActorChan || bPendingArchetypeSpawn) && InClass == UObject::StaticClass() && Obj != nullptr)
 		{
-			UUnitTestNetConnection* UnitConn = Cast<UUnitTestNetConnection>(GActiveReceiveUnitConnection);
-			bool bAllowActor = false;
+			bool bBlockActor = true;
 
-			if (UnitConn != nullptr)
+			if (MinClient == nullptr)
 			{
-				UMinimalClient* MinClient = UnitConn->MinClient;
+				MinClient = UMinimalClient::GetMinClientFromConn(GActiveReceiveUnitConnection);
+			}
 
-				if (MinClient != nullptr && MinClient->RepActorSpawnDel.IsBound())
-				{
-					bAllowActor = MinClient->RepActorSpawnDel.Execute(Obj->GetClass(), GIsInitializingActorChan);
-				}
+			if (MinClient != nullptr)
+			{
+				MinClient->RepActorSpawnDel.ExecuteIfBound(Obj->GetClass(), GIsInitializingActorChan, bBlockActor);
 			}
 
-			if (!bAllowActor)
+			if (bBlockActor)
 			{
 				UE_LOG(LogUnitTest, Log,
 						TEXT("Blocking replication/spawning of actor on client (add to NotifyAllowNetActor if required)."));
@@ -80,9 +79,6 @@ bool UUnitTestPackageMap::SerializeObject(FArchive& Ar, UClass* InClass, UObject
 						*Obj->GetFullName());
 
 				Obj = nullptr;
-
-				// NULL the control channel, to break code that would disconnect the client (control chan is recovered, in ReceivedBunch)
-				Connection->Channels[0] = nullptr;
 			}
 		}
 	}
diff --git a/Engine/Plugins/NetcodeUnitTest/NetcodeUnitTest/Source/NetcodeUnitTest/Private/NetcodeUnitTest.cpp b/Engine/Plugins/NetcodeUnitTest/NetcodeUnitTest/Source/NetcodeUnitTest/Private/NetcodeUnitTest.cpp
--- a/Engine/Plugins/NetcodeUnitTest/NetcodeUnitTest/Source/NetcodeUnitTest/Private/NetcodeUnitTest.cpp
+++ b/Engine/Plugins/NetcodeUnitTest/NetcodeUnitTest/Source/NetcodeUnitTest/Private/NetcodeUnitTest.cpp
@@ -16,16 +16,11 @@
  * Globals
  */
 
-UNetConnection* GActiveReceiveUnitConnection = NULL;
+UNetConnection* GActiveReceiveUnitConnection = nullptr;
 bool GIsInitializingActorChan = false;
-bool GSentBunch = false;
 
 ELogType GActiveLogTypeFlags = ELogType::None;
 
-#if !UE_BUILD_SHIPPING
-TMap<void*, InternalProcessEventCallback> ActiveProcessEventCallbacks;
-#endif
-
 
 /**
  * Definitions/implementations
@@ -122,7 +117,7 @@ public:
 	static void OnWorldCreated(UWorld* UnrealWorld, const UWorld::InitializationValues IVS)
 	{
 		// If NUTActor isn't already in RuntimeServerActors, add it now
-		if (GEngine != NULL)
+		if (GEngine != nullptr)
 		{
 			bool bNUTActorPresent = GEngine->RuntimeServerActors.ContainsByPredicate(
 				[](const FString& CurServerActor)
@@ -147,7 +142,7 @@ public:
 	}
 };
 
-FWorldDelegates::FWorldInitializationEvent::FDelegate FNetcodeUnitTest::OnWorldCreatedDelegate = NULL;
+FWorldDelegates::FWorldInitializationEvent::FDelegate FNetcodeUnitTest::OnWorldCreatedDelegate = nullptr;
 
 #if TARGET_UE4_CL >= CL_DEPRECATEDEL
 FDelegateHandle FNetcodeUnitTest::OnWorldCreatedDelegateHandle;
