Copying //UE4/Dev-Networking to //UE4/Dev-Main (Source: //UE4/Dev-Networking @ 3596412)

diff --git a/Engine/Plugins/NetcodeUnitTest/NetcodeUnitTest/Source/NetcodeUnitTest/Private/UnitTestManager.cpp b/Engine/Plugins/NetcodeUnitTest/NetcodeUnitTest/Source/NetcodeUnitTest/Private/UnitTestManager.cpp
--- a/Engine/Plugins/NetcodeUnitTest/NetcodeUnitTest/Source/NetcodeUnitTest/Private/UnitTestManager.cpp
+++ b/Engine/Plugins/NetcodeUnitTest/NetcodeUnitTest/Source/NetcodeUnitTest/Private/UnitTestManager.cpp
@@ -220,7 +220,7 @@ void UUnitTestManager::InitializeLogs()
 				}
 			}
 
-			bool Visit(const TCHAR* FilenameOrDirectory, bool bIsDirectory)
+			virtual bool Visit(const TCHAR* FilenameOrDirectory, bool bIsDirectory) override
 			{
 				if (bIsDirectory)
 				{
@@ -1285,7 +1285,7 @@ void UUnitTestManager::Tick(float DeltaTime)
 
 					// @todo #JohnB: Move NetTick to UClientUnitTest?
 					if ((CurTime - CurUnitTest->LastNetTick) > NetTickInterval && CurClientUnitTest != nullptr &&
-						CurClientUnitTest->UnitConn != nullptr)
+						CurClientUnitTest->MinClient != nullptr)
 					{
 						CurUnitTest->NetTick();
 						CurUnitTest->LastNetTick = CurTime;
@@ -1551,7 +1551,7 @@ bool UUnitTestManager::Exec(UWorld* InWorld, const TCHAR* Cmd, FOutputDevice& Ar
 	// Debug unit test commands
 	else if (UnitTestName == TEXT("debug"))
 	{
-		UUnitTest** TargetUnitTestRef = (InWorld != NULL ? ActiveUnitTests.FindByPredicate(
+		UUnitTest** TargetUnitTestRef = (InWorld != nullptr ? ActiveUnitTests.FindByPredicate(
 			[&InWorld](const UUnitTest* InElement)
 			{
 				auto CurUnitTest = Cast<UClientUnitTest>(InElement);
@@ -1559,18 +1559,18 @@ bool UUnitTestManager::Exec(UWorld* InWorld, const TCHAR* Cmd, FOutputDevice& Ar
 
 				return CurMinClient != nullptr && CurMinClient->GetUnitWorld() == InWorld;
 			})
-			: NULL);
+			: nullptr);
 
-		UClientUnitTest* TargetUnitTest = (TargetUnitTestRef != NULL ? Cast<UClientUnitTest>(*TargetUnitTestRef) : NULL);
+		UClientUnitTest* TargetUnitTest = (TargetUnitTestRef != nullptr ? Cast<UClientUnitTest>(*TargetUnitTestRef) : nullptr);
 
 		// Alternatively, if a unit test has not launched started connecting to a server, its world may not be setup,
 		// so can detect by checking the active log unit test too
-		if (TargetUnitTest == NULL && GActiveLogUnitTest != NULL)
+		if (TargetUnitTest == nullptr && GActiveLogUnitTest != nullptr)
 		{
 			TargetUnitTest = Cast<UClientUnitTest>(GActiveLogUnitTest);
 		}
 
-		if (TargetUnitTest != NULL)
+		if (TargetUnitTest != nullptr)
 		{
 			if (FParse::Command(&Cmd, TEXT("Requirements")))
 			{
@@ -1609,21 +1609,24 @@ bool UUnitTestManager::Exec(UWorld* InWorld, const TCHAR* Cmd, FOutputDevice& Ar
 			}
 			else if (FParse::Command(&Cmd, TEXT("ForceReady")))
 			{
-				if (TargetUnitTest != NULL && !!(TargetUnitTest->UnitTestFlags & EUnitTestFlags::LaunchServer) &&
-					TargetUnitTest->ServerHandle.IsValid() && TargetUnitTest->UnitPC == NULL)
+				if (TargetUnitTest != nullptr && !!(TargetUnitTest->UnitTestFlags & EUnitTestFlags::LaunchServer) &&
+					TargetUnitTest->ServerHandle.IsValid() && TargetUnitTest->UnitPC == nullptr)
 				{
 					Ar.Logf(TEXT("Forcing unit test '%s' as ready to connect client."), *TargetUnitTest->GetUnitTestName());
 
-					TargetUnitTest->ConnectFakeClient();
+					TargetUnitTest->ConnectMinimalClient();
 				}
 			}
 			else if (FParse::Command(&Cmd, TEXT("Disconnect")))
 			{
-				if (TargetUnitTest != NULL && TargetUnitTest->UnitConn != NULL)
+				UNetConnection* UnitConn = (TargetUnitTest != nullptr && TargetUnitTest->MinClient != nullptr ?
+											TargetUnitTest->MinClient->GetConn() : nullptr);
+
+				if (UnitConn != nullptr)
 				{
 					Ar.Logf(TEXT("Forcing unit test '%s' to disconnect."), *TargetUnitTest->GetUnitTestName());
 
-					TargetUnitTest->UnitConn->Close();
+					UnitConn->Close();
 				}
 			}
 		}
@@ -1762,10 +1765,10 @@ void UUnitTestManager::Serialize(const TCHAR* Data, ELogVerbosity::Type Verbosit
 				}
 
 
-				FString* LogLine = NULL;
+				FString* LogLine = nullptr;
 				UUnitTest* CurLogUnitTest = Cast<UUnitTest>(GActiveLogUnitTest);
 
-				if (CurLogUnitTest != NULL)
+				if (CurLogUnitTest != nullptr)
 				{
 					// Store the log within the unit test
 					CurLogUnitTest->StatusLogSummary.Add(MakeShareable(new FUnitStatusLog(CurLogType, FString(Data))));
@@ -1802,18 +1805,20 @@ void UUnitTestManager::Serialize(const TCHAR* Data, ELogVerbosity::Type Verbosit
 
 
 			ELogType CurLogType = ELogType::Local | GActiveLogTypeFlags;
-			UUnitTest* SourceUnitTest = NULL;
+			UUnitTest* SourceUnitTest = nullptr;
 
 			// If this log was triggered, while a unit test net connection was processing a packet, find and notify the unit test
-			if (GActiveReceiveUnitConnection != NULL)
+			if (GActiveReceiveUnitConnection != nullptr)
 			{
 				CurLogType |= ELogType::OriginNet;
 
 				for (auto CurUnitTest : ActiveUnitTests)
 				{
 					UClientUnitTest* CurClientUnitTest = Cast<UClientUnitTest>(CurUnitTest);
+					UNetConnection* UnitConn = (CurClientUnitTest != nullptr && CurClientUnitTest->MinClient != nullptr ?
+												CurClientUnitTest->MinClient->GetConn() : nullptr);
 
-					if (CurClientUnitTest != NULL && CurClientUnitTest->UnitConn == GActiveReceiveUnitConnection)
+					if (UnitConn != nullptr && UnitConn == GActiveReceiveUnitConnection)
 					{
 						SourceUnitTest = CurUnitTest;
 						break;
@@ -1821,19 +1826,19 @@ void UUnitTestManager::Serialize(const TCHAR* Data, ELogVerbosity::Type Verbosit
 				}
 			}
 			// If it was triggered from within a unit test log, also notify
-			else if (GActiveLogUnitTest != NULL)
+			else if (GActiveLogUnitTest != nullptr)
 			{
 				CurLogType |= ELogType::OriginUnitTest;
 				SourceUnitTest = Cast<UUnitTest>(GActiveLogUnitTest);
 			}
 			// If it was triggered within an engine event, within a unit test, again notify
-			else if (GActiveLogEngineEvent != NULL)
+			else if (GActiveLogEngineEvent != nullptr)
 			{
 				CurLogType |= ELogType::OriginEngine;
 				SourceUnitTest = Cast<UUnitTest>(GActiveLogEngineEvent);
 			}
 			// If it was triggered during UWorld::Tick, for the world assigned to a unit test, again find and notify
-			else if (GActiveLogWorld != NULL)
+			else if (GActiveLogWorld != nullptr)
 			{
 				CurLogType |= ELogType::OriginEngine;
 
@@ -1850,7 +1855,7 @@ void UUnitTestManager::Serialize(const TCHAR* Data, ELogVerbosity::Type Verbosit
 				}
 			}
 
-			if (SourceUnitTest != NULL && (CurLogType & ELogType::OriginMask) != ELogType::None)
+			if (SourceUnitTest != nullptr && (CurLogType & ELogType::OriginMask) != ELogType::None)
 			{
 				SourceUnitTest->NotifyLocalLog(CurLogType, Data, Verbosity, Category);
 			}
@@ -2055,18 +2060,15 @@ static bool UnitTestExec(UWorld* InWorld, const TCHAR* Cmd, FOutputDevice& Ar)
 		{
 			uint8* Data = (uint8*)PointerVal;
 
+			// NOTE: This case covers TArray's which are empty, and can be allocated (Data != nullptr) or unallocated (Data == nullptr)
 			if (Data != nullptr || DataLen == 0)
 			{
 				NUTDebug::LogHexDump(Data, DataLen);
 			}
-			else if (Data == nullptr)
+			else //if (Data == nullptr)
 			{
 				Ar.Logf(TEXT("Invalid Data parameter."));
 			}
-			else // if (DataLen == 0)
-			{
-				Ar.Logf(TEXT("Invalid DataLen parameter."));
-			}
 		}
 		else
 		{
@@ -2090,18 +2092,15 @@ static bool UnitTestExec(UWorld* InWorld, const TCHAR* Cmd, FOutputDevice& Ar)
 		{
 			uint8* Data = (uint8*)PointerVal;
 
+			// NOTE: This case covers TArray's which are empty, and can be allocated (Data != nullptr) or unallocated (Data == nullptr)
 			if (Data != nullptr || DataLen == 0)
 			{
 				NUTDebug::LogBitDump(Data, DataLen);
 			}
-			else if (Data == nullptr)
+			else //if (Data == nullptr)
 			{
 				Ar.Logf(TEXT("Invalid Data parameter."));
 			}
-			else // if (DataLen == 0)
-			{
-				Ar.Logf(TEXT("Invalid DataLen parameter."));
-			}
 		}
 		else
 		{
diff --git a/Engine/Plugins/NetcodeUnitTest/NetcodeUnitTest/Source/NetcodeUnitTest/Public/NUTUtil.h b/Engine/Plugins/NetcodeUnitTest/NetcodeUnitTest/Source/NetcodeUnitTest/Public/NUTUtil.h
--- a/Engine/Plugins/NetcodeUnitTest/NetcodeUnitTest/Source/NetcodeUnitTest/Public/NUTUtil.h
+++ b/Engine/Plugins/NetcodeUnitTest/NetcodeUnitTest/Source/NetcodeUnitTest/Public/NUTUtil.h
@@ -9,25 +9,14 @@
 #include "Engine/Engine.h"
 #include "Misc/OutputDeviceError.h"
 
+
 class UNetDriver;
 class UUnitTest;
 struct FStackTracker;
 class FOutputDeviceFile;
 
-// @todo #JohnBRefactor: Convert to multicast delegate
-
 // @todo #JohnBRefactor: Adjust all of these utility .h files, so that they implement code in .cpp, as these probably slow down compilation
 
-#if !UE_BUILD_SHIPPING
-// Process even callback, which passes an extra parameter, for identifying the origin of the hook
-typedef bool (*InternalProcessEventCallback)(AActor* /*Actor*/, UFunction* /*Function*/, void* /*Parameters*/, void* /*HookOrigin*/);
-
-
-/** Track callbacks not by setting the global extern/hook above, but by indirectly handling multiple callbacks locally */
-extern TMap<void*, InternalProcessEventCallback> ActiveProcessEventCallbacks;
-#endif
-
-
 
 /**
  * Output device for allowing quick/dynamic creation of a customized output device, using lambda's passed to delegates
@@ -220,58 +209,6 @@ public:
 };
 
 
-#if !UE_BUILD_SHIPPING
-static bool HandleProcessEventCallback(AActor* Actor, UFunction* Function, void* Parameters)
-{
-	bool bBlockEvent = false;
-
-	for (auto It = ActiveProcessEventCallbacks.CreateConstIterator(); It; ++It)
-	{
-		InternalProcessEventCallback CurCallback = It->Value;
-
-		// Replace 'Parameters' with the TMap void*, which refers to the object that set the callback
-		if (CurCallback(Actor, Function, Parameters, It->Key))
-		{
-			bBlockEvent = true;
-		}
-	}
-
-	return bBlockEvent;
-}
-
-// HookOrigin should reference the object responsible for the callback, for identification
-static void AddProcessEventCallback(void* HookOrigin, InternalProcessEventCallback InCallback)
-{
-	if (HookOrigin != NULL)
-	{
-		if (ActiveProcessEventCallbacks.Num() == 0)
-		{
-			AActor::ProcessEventDelegate.BindStatic(&HandleProcessEventCallback);
-		}
-
-		ActiveProcessEventCallbacks.Add(HookOrigin, InCallback);
-	}
-	else
-	{
-		UE_LOG(LogUnitTest, Log, TEXT("AddProcessEventCallback: HookOrigin == NULL"));
-	}
-}
-
-static void RemoveProcessEventCallback(void* HookOrigin, InternalProcessEventCallback InCallback)
-{
-	if (HookOrigin != NULL)
-	{
-		ActiveProcessEventCallbacks.Remove(HookOrigin);
-
-		if (ActiveProcessEventCallbacks.Num() == 0)
-		{
-			AActor::ProcessEventDelegate.Unbind();
-		}
-	}
-}
-#endif
-
-
 // Hack for accessing private members in various Engine classes
 
 // This template, is used to link an arbitrary class member, to the GetPrivate function
diff --git a/Engine/Plugins/NetcodeUnitTest/NetcodeUnitTest/Source/NetcodeUnitTest/Public/NUTUtilDebug.h b/Engine/Plugins/NetcodeUnitTest/NetcodeUnitTest/Source/NetcodeUnitTest/Public/NUTUtilDebug.h
--- a/Engine/Plugins/NetcodeUnitTest/NetcodeUnitTest/Source/NetcodeUnitTest/Public/NUTUtilDebug.h
+++ b/Engine/Plugins/NetcodeUnitTest/NetcodeUnitTest/Source/NetcodeUnitTest/Public/NUTUtilDebug.h
@@ -133,6 +133,10 @@ public:
 };
 
 
+// @todo #JohnB: When you continue implementing this, as a part of the ProcessEvent stack trace feature below,
+//					merge this class with the very similar 'FProcessEventHook' class in NUTUtilNet.h,
+//					then make the stack trace hook use that.
+#if 0
 #if !UE_BUILD_SHIPPING
 /**
  * Base class for transparently hooking ProcessEvent
@@ -162,8 +166,13 @@ protected:
 	/** If a 'Actor::ProcessEventDelegate' value was already set, this caches it so it can be transparently hooked and restored later */
 	FOnProcessEvent		OrigEventHook;
 };
+#endif
+#endif
 
 
+
+// @todo #JohnB: Reimplement this, by refactoring the above commented class, when this debug feature is next needed
+#if 0
 /**
  * A class for hooking and logging all ProcessEvent calls, within a particular code scope, similar to the above code
  */
@@ -427,7 +436,7 @@ public:
 	/**
 	 * Destructor
 	 */
-	~FLogStackTraceManager()
+	virtual ~FLogStackTraceManager() override
 	{
 		if (GLog != nullptr)
 		{
diff --git a/Engine/Plugins/NetcodeUnitTest/NetcodeUnitTest/Source/NetcodeUnitTest/Public/Net/NUTUtilNet.h b/Engine/Plugins/NetcodeUnitTest/NetcodeUnitTest/Source/NetcodeUnitTest/Public/Net/NUTUtilNet.h
--- a/Engine/Plugins/NetcodeUnitTest/NetcodeUnitTest/Source/NetcodeUnitTest/Public/Net/NUTUtilNet.h
+++ b/Engine/Plugins/NetcodeUnitTest/NetcodeUnitTest/Source/NetcodeUnitTest/Public/Net/NUTUtilNet.h
@@ -5,7 +5,7 @@
 #include "CoreMinimal.h"
 #include "Engine/Channel.h"
 #include "NetcodeUnitTest.h"
-#include "Engine/PendingNetGame.h"
+#include "Engine/NetworkDelegates.h"
 #include "Engine/World.h"
 #include "Sockets.h"
 
@@ -22,157 +22,96 @@ struct FUniqueNetIdRepl;
 class UClientUnitTest;
 
 
+
+// Delegates
+
 /**
- * A transparent FSocket-hook, for capturing lowest-level socket operations (use Attach/Detach functions for hooking)
+ * Delegate used to control/block receiving of RPC's
  *
- * NOTE: Not compatible with all net subsystems, e.g. some subsystems expect a specific FSocket-subclass, and statically cast as such.
- * NOTE: Best if only temporarily attached, or the original socket may leak (no attempt is made to handle destruction properly).
+ * @param Actor			The actor the RPC will be called on
+ * @param Function		The RPC to call
+ * @param Parameters	The parameters data blob
+ * @param bBlockRPC		Whether or not to block execution of the RPC
+*/
+DECLARE_DELEGATE_FourParams(FOnProcessNetEvent, AActor* /*Actor*/, UFunction* /*Function*/, void* /*Parameters*/, bool& /*bBlockRPC*/);
+
+
+/**
+ * Class for encapsulating ProcessEvent and general RPC hooks, implemented globally for each UWorld
  */
-class FSocketHook : public FSocket
+class FProcessEventHook
 {
 public:
 	/**
-	 * Base constructor
+	 * Get a reference to the net event hook singular
 	 */
-	FSocketHook()
-		: MinClient(nullptr)
-		, SendToDel()
-		, bLastSendToBlocked(false)
-		, HookedSocket(nullptr)
+	static FProcessEventHook& Get()
 	{
+		static FProcessEventHook* HookSingular = nullptr;
+
+		if (HookSingular == nullptr)
+		{
+			HookSingular = new FProcessEventHook();
+		}
+
+		return *HookSingular;
 	}
 
+public:
 	/**
-	 * Base constructor
+	 * Default constructor
 	 */
-	FSocketHook(FSocket* InHookedSocket)
-		: FSocket((InHookedSocket != nullptr ? InHookedSocket->GetSocketType() : ESocketType::SOCKTYPE_Datagram),
-					(InHookedSocket != nullptr ? InHookedSocket->GetDescription() + TEXT(" (hooked)") : TEXT("(hooked)")))
-		, MinClient(nullptr)
-		, SendToDel()
-		, bLastSendToBlocked(false)
-		, HookedSocket(InHookedSocket)
+	FProcessEventHook()
+		: NetEventHooks()
 	{
 	}
 
 	/**
-	 * Attach the hook to the specified FSocket* variable
+	 * Adds an RPC hook for the specified UWorld.
 	 *
-	 * @param InSocketVar	The socket var to hook
+	 * @param InWorld	The UWorld whose RPC's should be hooked
+	 * @param InHook	The delegate for handling the hooked RPC's
 	 */
-	FORCEINLINE void Attach(FSocket*& InSocketVar)
-	{
-		check(InSocketVar != this);
-		check(HookedSocket == nullptr);
-
-		HookedSocket = InSocketVar;
-		InSocketVar = this;
-	}
+	void AddRPCHook(UWorld* InWorld, FOnProcessNetEvent InHook);
 
 	/**
-	 * Detach the hook from the specified FSocket* variable
+	 * Removes the RPC hook for the specified UWorld.
 	 *
-	 * @param InSocketVar	The socket var to detach from 
+	 * @param InWorld	The UWorld whose RPC hook should be cleared.
 	 */
-	FORCEINLINE void Detach(FSocket*& InSocketVar)
-	{
-		check(InSocketVar == this);
-		check(HookedSocket != nullptr);
-
-		InSocketVar = HookedSocket;
-		HookedSocket = nullptr;
-	}
-
-	virtual ~FSocketHook()
-	{
-	}
-
-
-	virtual bool Close() override;
-
-	virtual bool Bind(const FInternetAddr& Addr) override;
-
-	virtual bool Connect(const FInternetAddr& Addr) override;
-
-	virtual bool Listen(int32 MaxBacklog) override;
-
-	virtual bool HasPendingConnection(bool& bHasPendingConnection) override;
-
-	virtual bool HasPendingData(uint32& PendingDataSize) override;
-
-	virtual class FSocket* Accept(const FString& InSocketDescription) override;
-
-	virtual class FSocket* Accept(FInternetAddr& OutAddr, const FString& InSocketDescription) override;
-
-	virtual bool SendTo(const uint8* Data, int32 Count, int32& BytesSent, const FInternetAddr& Destination) override;
-
-	virtual bool Send(const uint8* Data, int32 Count, int32& BytesSent) override;
-
-	virtual bool RecvFrom(uint8* Data, int32 BufferSize, int32& BytesRead, FInternetAddr& Source,
-							ESocketReceiveFlags::Type Flags) override;
-
-	virtual bool Recv(uint8* Data, int32 BufferSize, int32& BytesRead, ESocketReceiveFlags::Type Flags) override;
-
-	virtual bool Wait(ESocketWaitConditions::Type Condition, FTimespan WaitTime) override;
-
-	virtual ESocketConnectionState GetConnectionState() override;
-
-	virtual void GetAddress(FInternetAddr& OutAddr) override;
-
-	virtual bool GetPeerAddress(FInternetAddr& OutAddr) override;
-
-	virtual bool SetNonBlocking(bool bIsNonBlocking) override;
-
-	virtual bool SetBroadcast(bool bAllowBroadcast) override;
-
-	virtual bool JoinMulticastGroup(const FInternetAddr& GroupAddress) override;
-
-	virtual bool LeaveMulticastGroup(const FInternetAddr& GroupAddress) override;
-
-	virtual bool SetMulticastLoopback(bool bLoopback) override;
-
-	virtual bool SetMulticastTtl(uint8 TimeToLive) override;
-
-	virtual bool SetReuseAddr(bool bAllowReuse) override;
+	void RemoveRPCHook(UWorld* InWorld);
 
-	virtual bool SetLinger(bool bShouldLinger, int32 Timeout) override;
-
-	virtual bool SetRecvErr(bool bUseErrorQueue) override;
-
-	virtual bool SetSendBufferSize(int32 Size, int32& NewSize) override;
-
-	virtual bool SetReceiveBufferSize(int32 Size, int32& NewSize) override;
-
-	virtual int32 GetPortNo() override;
-
-
-public:
 	/**
-	 * Delegate for hooking 'SendTo'
+	 * Adds a non-RPC hook for the specified UWorld.
 	 *
-	 * @param Data			The data being sent
-	 * @param Count			The number of bytes being sent
-	 * @param bBlockSend	Whether or not to block the send (defaults to false)
+	 * @param InWorld	The UWorld whose RPC's should be hooked
+	 * @param InHook	The delegate for handling the hooked RPC's
 	 */
-	DECLARE_DELEGATE_ThreeParams(FSendToDel, void* /*Data*/, int32 /*Count*/, bool& /*bBlockSend*/);
-
+	void AddEventHook(UWorld* InWorld, FOnProcessNetEvent InHook);
 
-public:
-	/** The minimal client which owns this socket hook (if any) */
-	UMinimalClient* MinClient;
+	/**
+	 * Removes the hook for the specified UWorld.
+	 *
+	 * @param InWorld	The UWorld whose RPC hook should be cleared.
+	 */
+	void RemoveEventHook(UWorld* InWorld);
 
-	/** Delegate for hooking 'SendTo' */
-	FSendToDel SendToDel;
 
-	/** Whether or not the last 'SendTo' call was blocked */
-	bool bLastSendToBlocked;
+private:
+	/**
+	 * Base hook for AActor::ProcessEventDelegate - responsible for filtering based on RPC events and Actor UWorld
+	 */
+	bool HandleProcessEvent(AActor* Actor, UFunction* Function, void* Parameters);
 
+private:
+	/** The global list of RPC hooks, and the UWorld they are associated with */
+	TMap<UWorld*, FOnProcessNetEvent> NetEventHooks;
 
-public:
-	/** The original socket which has been hooked */
-	FSocket* HookedSocket;
+	/** The global list of ProcessEvent hooks, and the UWorld they are associated with */
+	TMap<UWorld*, FOnProcessNetEvent> EventHooks;
 };
 
+
 /**
  * A delegate network notify class, to allow for easy inline-hooking.
  * 
@@ -334,72 +273,6 @@ private:
  */
 struct NUTNet
 {
-	/**
-	 * Gets the index of the first free/unused net channel, on the specified net connection
-	 *
-	 * @param InConnection	The net connection to get the index from
-	 * @param InChIndex		Optionally, specifies the channel index at which to begin searching
-	 * @return				Returns the first free channel index
-	 */
-	static int32 GetFreeChannelIndex(UNetConnection* InConnection, int32 InChIndex=INDEX_NONE);
-
-	/**
-	 * Creates a unit test channel, for capturing incoming data on a particular channel index,
-	 * and routing it to a delegate for hooking or other purposes
-	 *
-	 * @param InConnection	The connection to create the channel in
-	 * @param InType		The type of channel to create
-	 * @param InChIndex		Optionally, specify the channel index
-	 * @param bInVerifyOpen	Optionally, make the channel verify that it has been opened (keep sending initial packet, until ack)
-	 * @return				Returns the created unit test channel
-	 */
-	static UUnitTestChannel* CreateUnitTestChannel(UNetConnection* InConnection, EChannelType InType, int32 InChIndex=INDEX_NONE,
-													bool bInVerifyOpen=false);
-
-	/**
-	 * Outputs a bunch for the specified channel, with the ability to create the channel as well
-	 * NOTE: BunchSequence should start at 0 for new channels, and should be tracked per-channel
-	 * WARNING: Can return null! (e.g. if the control channel is saturated)
-	 *
-	 * @param BunchSequence		Specify a variable for keeping track of the bunch sequence on this channel (may need to initialize as 0)
-	 * @param InConnection		Specify the connection the bunch will be sent on
-	 * @param InChType			Specify the type of the channel the bunch will be sent on
-	 * @param InChIndex			Optionally, specify the index of the channel to send the bunch on
-	 * @param bGetNextFreeChan	Whether or not to pick the next free/unused channel, for this bunch
-	 * @return					Returns the created bunch
-	 */
-	static NETCODEUNITTEST_API FOutBunch* CreateChannelBunch(int32& BunchSequence, UNetConnection* InConnection, EChannelType InChType,
-												int32 InChIndex=INDEX_NONE, bool bGetNextFreeChan=false);
-
-	// @todo #JohnB: Deprecate this, used during transition to minimal client code
-#if 1
-	static FORCEINLINE FOutBunch* CreateChannelBunch(int32* BunchSequence, UNetConnection* InConnection, EChannelType InChType,
-												int32 InChIndex=INDEX_NONE, bool bGetNextFreeChan=false)
-	{
-		check(BunchSequence != nullptr);
-
-		return CreateChannelBunch(*BunchSequence, InConnection, InChType, InChIndex, bGetNextFreeChan);
-	}
-#endif
-
-	/**
-	 * Sends a control channel bunch over a UUnitTestChannel based control channel
-	 *
-	 * @param InConnection		The connection to send the bunch over
-	 * @param ControlChanBunch	The bunch to send over the connections control channel
-	 */
-	static NETCODEUNITTEST_API void SendControlBunch(UNetConnection* InConnection, FOutBunch& ControlChanBunch);
-
-
-	/**
-	 * Creates an instance of the unit test net driver (can be used to create multiple such drivers, for creating multiple connections)
-	 *
-	 * @param InWorld	The world that the net driver should be associated with
-	 * @return			Returns the created net driver instance
-	 */
-	static NETCODEUNITTEST_API UUnitTestNetDriver* CreateUnitTestNetDriver(UWorld* InWorld);
-
-
 	/**
 	 * Handles setting up the client beacon once it is replicated, so that it can properly send RPC's
 	 * (normally the serverside client beacon links up with the pre-existing beacon on the clientside,
@@ -408,7 +281,7 @@ struct NUTNet
 	 * @param InBeacon		The beacon that should be setup
 	 * @param InConnection	The connection associated with the beacon
 	 */
-	static void HandleBeaconReplicate(class AOnlineBeaconClient* InBeacon, UNetConnection* InConnection);
+	static void HandleBeaconReplicate(AActor* InBeacon, UNetConnection* InConnection);
 
 
 	/**
diff --git a/Engine/Plugins/NetcodeUnitTest/NetcodeUnitTest/Source/NetcodeUnitTest/Public/NetcodeUnitTest.h b/Engine/Plugins/NetcodeUnitTest/NetcodeUnitTest/Source/NetcodeUnitTest/Public/NetcodeUnitTest.h
--- a/Engine/Plugins/NetcodeUnitTest/NetcodeUnitTest/Source/NetcodeUnitTest/Public/NetcodeUnitTest.h
+++ b/Engine/Plugins/NetcodeUnitTest/NetcodeUnitTest/Source/NetcodeUnitTest/Public/NetcodeUnitTest.h
@@ -163,9 +163,6 @@ extern UNetConnection* GActiveReceiveUnitConnection;
 /** Whether not an actor channel, is in the process of initializing the remote actor */
 extern bool GIsInitializingActorChan;
 
-/** Set whenever a bunch is sent; used to test RPC sends, by setting this to false, calling the RPC, and then testing the value */
-extern bool GSentBunch;
-
 /** The current ELogType flag modifiers, associated with a UNIT_LOG or UNIT_STATUS_LOG/STATUS_LOG call */
 extern NETCODEUNITTEST_API ELogType GActiveLogTypeFlags;
 
@@ -211,6 +208,7 @@ inline ELogType OptionalFlags(ELogType InFlags=ELogType::None)
 		UE_LOG(LogUnitTest, Error, TEXT(PREPROCESSOR_TO_STRING(__FILE__)) TEXT(": Line: ") TEXT(PREPROCESSOR_TO_STRING(__LINE__)) \
 				TEXT(":")); \
 		UE_LOG(LogUnitTest, Error, TEXT("Condition '(") TEXT(PREPROCESSOR_TO_STRING(Condition)) TEXT(")' failed")); \
+		check(false); /* Try to get a meaningful stack trace. */ \
 		FPlatformMisc::RequestExit(true); \
 		CA_ASSUME(false); \
 	}
@@ -221,7 +219,7 @@ inline ELogType OptionalFlags(ELogType InFlags=ELogType::None)
 #define UNIT_LOG_BEGIN(UnitTestObj, UnitLogTypeFlags) \
 	{ \
 		UUnitTestBase* OldUnitLogVal = GActiveLogUnitTest; \
-		GActiveLogUnitTest = (UnitTestObj != NULL ? UnitTestObj : GActiveLogUnitTest); \
+		GActiveLogUnitTest = (UnitTestObj != nullptr ? UnitTestObj : GActiveLogUnitTest); \
 		GActiveLogTypeFlags = ELogType::UnitStatus | OptionalFlags(UnitLogTypeFlags);
 
 #define UNIT_LOG_END() \
@@ -295,8 +293,8 @@ inline ELogType OptionalFlags(ELogType InFlags=ELogType::None)
 #define UNIT_EVENT_CLEAR \
 	UUnitTestBase* OldEventLogVal = GActiveLogEngineEvent; \
 	UWorld* OldEventLogWorld = GActiveLogWorld; \
-	GActiveLogEngineEvent = NULL; \
-	GActiveLogWorld = NULL;
+	GActiveLogEngineEvent = nullptr; \
+	GActiveLogWorld = nullptr;
 
 /**
  * Restores all stored/cleared engine-log-capture events
@@ -332,7 +330,7 @@ inline ELogType OptionalFlags(ELogType InFlags=ELogType::None)
 		UNIT_EVENT_CLEAR; \
 		if ((GActiveLogTypeFlags & ELogType::StatusError) == ELogType::StatusError) \
 		{ \
-			if (SourceUnitTest != NULL) \
+			if (SourceUnitTest != nullptr) \
 			{ \
 				UE_LOG(LogUnitTest, Error, TEXT("%s: %s"), *SourceUnitTest->GetUnitTestName(), *FString::Printf(Format, ##__VA_ARGS__)); \
 			} \
@@ -343,7 +341,7 @@ inline ELogType OptionalFlags(ELogType InFlags=ELogType::None)
 		} \
 		else if ((GActiveLogTypeFlags & ELogType::StatusWarning) == ELogType::StatusWarning) \
 		{ \
-			if (SourceUnitTest != NULL) \
+			if (SourceUnitTest != nullptr) \
 			{ \
 				UE_LOG(LogUnitTest, Warning, TEXT("%s: %s"), *SourceUnitTest->GetUnitTestName(), \
 						*FString::Printf(Format, ##__VA_ARGS__)); \
@@ -362,7 +360,7 @@ inline ELogType OptionalFlags(ELogType InFlags=ELogType::None)
 		GActiveLogTypeFlags = ELogType::None; \
 	}
 
-#define STATUS_LOG(UnitLogTypeFlags, Format, ...) STATUS_LOG_OBJ(NULL, UnitLogTypeFlags, Format, ##__VA_ARGS__)
+#define STATUS_LOG(UnitLogTypeFlags, Format, ...) STATUS_LOG_OBJ(nullptr, UnitLogTypeFlags, Format, ##__VA_ARGS__)
 
 /**
  * Version of the above, for unit test status window entries, which are from specific unit tests
diff --git a/Engine/Plugins/NetcodeUnitTest/NetcodeUnitTest/VAXSnippets/Unit Test Source.txt b/Engine/Plugins/NetcodeUnitTest/NetcodeUnitTest/VAXSnippets/Unit Test Source.txt
--- a/Engine/Plugins/NetcodeUnitTest/NetcodeUnitTest/VAXSnippets/Unit Test Source.txt	
+++ b/Engine/Plugins/NetcodeUnitTest/NetcodeUnitTest/VAXSnippets/Unit Test Source.txt	
@@ -7,9 +7,8 @@ Source:
 
 #include "UnitTests/$FILE_BASE$.h"
 
-#include "$UnitTestProject$PCH.h"
-
 #include "UnitTestEnvironment.h"
+#include "MinimalClient.h"
 
 
 /**
@@ -33,7 +32,7 @@ U$FILE_BASE$::U$FILE_BASE$(const FObjectInitializer& ObjectInitializer)
 	UnitTestTimeout = 60;
 
 
-	UnitTestFlags |= (EUnitTestFlags::LaunchServer | $end$);
+	SetFlags<EUnitTestFlags::LaunchServer | $end$,>();
 }
 
 void U$FILE_BASE$::InitializeEnvironmentSettings()
diff --git a/Engine/Plugins/Online/OnlineSubsystemOculus/Source/Private/OculusNetConnection.cpp b/Engine/Plugins/Online/OnlineSubsystemOculus/Source/Private/OculusNetConnection.cpp
--- a/Engine/Plugins/Online/OnlineSubsystemOculus/Source/Private/OculusNetConnection.cpp
+++ b/Engine/Plugins/Online/OnlineSubsystemOculus/Source/Private/OculusNetConnection.cpp
@@ -104,7 +104,13 @@ void UOculusNetConnection::LowLevelSend(void* Data, int32 CountBytes, int32 Coun
 		}
 	}
 
-	if (CountBytes > 0)
+	bool bBlockSend = false;
+
+#if !UE_BUILD_SHIPPING
+	LowLevelSendDel.ExecuteIfBound((void*)DataToSend, CountBytes, bBlockSend);
+#endif
+
+	if (!bBlockSend && CountBytes > 0)
 	{
 		ovr_Net_SendPacket(PeerID, static_cast<size_t>(CountBytes), DataToSend, (InternalAck) ? ovrSend_Reliable : ovrSend_Unreliable);
 	}
diff --git a/Engine/Plugins/Online/OnlineSubsystemUtils/Source/OnlineSubsystemUtils/Private/IpConnection.cpp b/Engine/Plugins/Online/OnlineSubsystemUtils/Source/OnlineSubsystemUtils/Private/IpConnection.cpp
--- a/Engine/Plugins/Online/OnlineSubsystemUtils/Source/OnlineSubsystemUtils/Private/IpConnection.cpp
+++ b/Engine/Plugins/Online/OnlineSubsystemUtils/Source/OnlineSubsystemUtils/Private/IpConnection.cpp
@@ -151,25 +151,34 @@ void UIpConnection::LowLevelSend(void* Data, int32 CountBytes, int32 CountBits)
 		}
 	}
 
-	// Send to remote.
-	int32 BytesSent = 0;
-	CLOCK_CYCLES(Driver->SendCycles);
+	bool bBlockSend = false;
 
-	if ( CountBytes > MaxPacket )
+#if !UE_BUILD_SHIPPING
+	LowLevelSendDel.ExecuteIfBound((void*)DataToSend, CountBytes, bBlockSend);
+#endif
+
+	if (!bBlockSend)
 	{
-		UE_LOG( LogNet, Warning, TEXT( "UIpConnection::LowLevelSend: CountBytes > MaxPacketSize! Count: %i, MaxPacket: %i %s" ), CountBytes, MaxPacket, *Describe() );
-	}
+		// Send to remote.
+		int32 BytesSent = 0;
+		CLOCK_CYCLES(Driver->SendCycles);
 
-	FPacketAudit::NotifyLowLevelSend((uint8*)DataToSend, CountBytes, CountBits);
+		if ( CountBytes > MaxPacket )
+		{
+			UE_LOG( LogNet, Warning, TEXT( "UIpConnection::LowLevelSend: CountBytes > MaxPacketSize! Count: %i, MaxPacket: %i %s" ), CountBytes, MaxPacket, *Describe() );
+		}
 
-	if (CountBytes > 0)
-	{
-		Socket->SendTo(DataToSend, CountBytes, BytesSent, *RemoteAddr);
-	}
+		FPacketAudit::NotifyLowLevelSend((uint8*)DataToSend, CountBytes, CountBits);
 
-	UNCLOCK_CYCLES(Driver->SendCycles);
-	NETWORK_PROFILER(GNetworkProfiler.FlushOutgoingBunches(this));
-	NETWORK_PROFILER(GNetworkProfiler.TrackSocketSendTo(Socket->GetDescription(),DataToSend,BytesSent,NumPacketIdBits,NumBunchBits,NumAckBits,NumPaddingBits,this));
+		if (CountBytes > 0)
+		{
+			Socket->SendTo(DataToSend, CountBytes, BytesSent, *RemoteAddr);
+		}
+
+		UNCLOCK_CYCLES(Driver->SendCycles);
+		NETWORK_PROFILER(GNetworkProfiler.FlushOutgoingBunches(this));
+		NETWORK_PROFILER(GNetworkProfiler.TrackSocketSendTo(Socket->GetDescription(),DataToSend,BytesSent,NumPacketIdBits,NumBunchBits,NumAckBits,NumPaddingBits,this));
+	}
 }
 
 FString UIpConnection::LowLevelGetRemoteAddress(bool bAppendPort)
diff --git a/Engine/Plugins/Online/OnlineSubsystemUtils/Source/OnlineSubsystemUtils/Private/IpNetDriver.cpp b/Engine/Plugins/Online/OnlineSubsystemUtils/Source/OnlineSubsystemUtils/Private/IpNetDriver.cpp
--- a/Engine/Plugins/Online/OnlineSubsystemUtils/Source/OnlineSubsystemUtils/Private/IpNetDriver.cpp
+++ b/Engine/Plugins/Online/OnlineSubsystemUtils/Source/OnlineSubsystemUtils/Private/IpNetDriver.cpp
@@ -529,66 +529,75 @@ void UIpNetDriver::LowLevelSend(FString Address, void* Data, int32 CountBits)
 
 void UIpNetDriver::ProcessRemoteFunction(class AActor* Actor, UFunction* Function, void* Parameters, FOutParmRec* OutParms, FFrame* Stack, class UObject* SubObject )
 {
-	bool bIsServer = IsServer();
+#if !UE_BUILD_SHIPPING
+	bool bBlockSendRPC = false;
+
+	SendRPCDel.ExecuteIfBound(Actor, Function, Parameters, OutParms, Stack, SubObject, bBlockSendRPC);
 
-	UNetConnection* Connection = NULL;
-	if (bIsServer)
+	if (!bBlockSendRPC)
+#endif
 	{
-		if ((Function->FunctionFlags & FUNC_NetMulticast))
+		bool bIsServer = IsServer();
+
+		UNetConnection* Connection = NULL;
+		if (bIsServer)
 		{
-			// Multicast functions go to every client
-			TArray<UNetConnection*> UniqueRealConnections;
-			for (int32 i=0; i<ClientConnections.Num(); ++i)
+			if ((Function->FunctionFlags & FUNC_NetMulticast))
 			{
-				Connection = ClientConnections[i];
-				if (Connection && Connection->ViewTarget)
+				// Multicast functions go to every client
+				TArray<UNetConnection*> UniqueRealConnections;
+				for (int32 i=0; i<ClientConnections.Num(); ++i)
 				{
-					// Do relevancy check if unreliable.
-					// Reliables will always go out. This is odd behavior. On one hand we wish to guarantee "reliables always get there". On the other
-					// hand, replicating a reliable to something on the other side of the map that is non relevant seems weird. 
-					//
-					// Multicast reliables should probably never be used in gameplay code for actors that have relevancy checks. If they are, the 
-					// rpc will go through and the channel will be closed soon after due to relevancy failing.
-
-					bool IsRelevant = true;
-					if ((Function->FunctionFlags & FUNC_NetReliable) == 0)
-					{
-						FNetViewer Viewer(Connection, 0.f);
-						IsRelevant = Actor->IsNetRelevantFor(Viewer.InViewer, Viewer.ViewTarget, Viewer.ViewLocation);
-					}
-					
-					if (IsRelevant)
+					Connection = ClientConnections[i];
+					if (Connection && Connection->ViewTarget)
 					{
-						if (Connection->GetUChildConnection() != NULL)
+						// Do relevancy check if unreliable.
+						// Reliables will always go out. This is odd behavior. On one hand we wish to guarantee "reliables always get there". On the other
+						// hand, replicating a reliable to something on the other side of the map that is non relevant seems weird. 
+						//
+						// Multicast reliables should probably never be used in gameplay code for actors that have relevancy checks. If they are, the 
+						// rpc will go through and the channel will be closed soon after due to relevancy failing.
+
+						bool IsRelevant = true;
+						if ((Function->FunctionFlags & FUNC_NetReliable) == 0)
 						{
-							Connection = ((UChildConnection*)Connection)->Parent;
+							FNetViewer Viewer(Connection, 0.f);
+							IsRelevant = Actor->IsNetRelevantFor(Viewer.InViewer, Viewer.ViewTarget, Viewer.ViewLocation);
 						}
+					
+						if (IsRelevant)
+						{
+							if (Connection->GetUChildConnection() != NULL)
+							{
+								Connection = ((UChildConnection*)Connection)->Parent;
+							}
 						
-						InternalProcessRemoteFunction( Actor, SubObject, Connection, Function, Parameters, OutParms, Stack, bIsServer );
+							InternalProcessRemoteFunction( Actor, SubObject, Connection, Function, Parameters, OutParms, Stack, bIsServer );
+						}
 					}
-				}
-			}			
+				}			
 
-			// Replicate any RPCs to the replay net driver so that they can get saved in network replays
-			UNetDriver* NetDriver = GEngine->FindNamedNetDriver(GetWorld(), NAME_DemoNetDriver);
-			if (NetDriver)
-			{
-				NetDriver->ProcessRemoteFunction(Actor, Function, Parameters, OutParms, Stack, SubObject);
+				// Replicate any RPCs to the replay net driver so that they can get saved in network replays
+				UNetDriver* NetDriver = GEngine->FindNamedNetDriver(GetWorld(), NAME_DemoNetDriver);
+				if (NetDriver)
+				{
+					NetDriver->ProcessRemoteFunction(Actor, Function, Parameters, OutParms, Stack, SubObject);
+				}
+				// Return here so we don't call InternalProcessRemoteFunction again at the bottom of this function
+				return;
 			}
-			// Return here so we don't call InternalProcessRemoteFunction again at the bottom of this function
-			return;
 		}
-	}
 
-	// Send function data to remote.
-	Connection = Actor->GetNetConnection();
-	if (Connection)
-	{
-		InternalProcessRemoteFunction( Actor, SubObject, Connection, Function, Parameters, OutParms, Stack, bIsServer );
-	}
-	else
-	{
-		UE_LOG(LogNet, Warning, TEXT("UIpNetDriver::ProcesRemoteFunction: No owning connection for actor %s. Function %s will not be processed."), *Actor->GetName(), *Function->GetName());
+		// Send function data to remote.
+		Connection = Actor->GetNetConnection();
+		if (Connection)
+		{
+			InternalProcessRemoteFunction( Actor, SubObject, Connection, Function, Parameters, OutParms, Stack, bIsServer );
+		}
+		else
+		{
+			UE_LOG(LogNet, Warning, TEXT("UIpNetDriver::ProcessRemoteFunction: No owning connection for actor %s. Function %s will not be processed."), *Actor->GetName(), *Function->GetName());
+		}
 	}
 }
 
diff --git a/Engine/Plugins/Online/OnlineSubsystemUtils/Source/OnlineSubsystemUtils/Private/OnlineBeaconClient.cpp b/Engine/Plugins/Online/OnlineSubsystemUtils/Source/OnlineSubsystemUtils/Private/OnlineBeaconClient.cpp
--- a/Engine/Plugins/Online/OnlineSubsystemUtils/Source/OnlineSubsystemUtils/Private/OnlineBeaconClient.cpp
+++ b/Engine/Plugins/Online/OnlineSubsystemUtils/Source/OnlineSubsystemUtils/Private/OnlineBeaconClient.cpp
@@ -256,10 +256,16 @@ void AOnlineBeaconClient::NotifyControlMessage(UNetConnection* Connection, uint8
 #if SUPPORTS_ENCRYPTION_TOKEN
 		case NMT_EncryptionAck:
 			{
-				UGameInstance* const Instance = GetGameInstance();
-				if (Instance)
+				if (FNetDelegates::OnReceivedNetworkEncryptionAck.IsBound())
 				{
-					Instance->ReceivedNetworkEncryptionAck(Connection);
+					TWeakObjectPtr<UNetConnection> WeakConnection = Connection;
+					FNetDelegates::OnReceivedNetworkEncryptionAck.Execute(FOnEncryptionKeyResponse::CreateUObject(this, &ThisClass::FinalizeEncryptedConnection, WeakConnection));
+				}
+				else
+				{
+					// Force close the session
+					UE_LOG(LogBeacon, Warning, TEXT("%s: No delegate available to handle encryption ack, disconnecting."), *Connection->GetName());
+					OnFailure();
 				}
 				break;
 			}
@@ -352,3 +358,34 @@ void AOnlineBeaconClient::NotifyControlMessage(UNetConnection* Connection, uint8
 		}
 	}	
 }
+
+void AOnlineBeaconClient::FinalizeEncryptedConnection(const FEncryptionKeyResponse& Response, TWeakObjectPtr<UNetConnection> WeakConnection)
+{
+	UNetConnection* Connection = WeakConnection.Get();
+	if (Connection)
+	{
+		if (Connection->State != USOCK_Invalid && Connection->State != USOCK_Closed && Connection->Driver)
+		{
+			if (Response.Response == EEncryptionResponse::Success)
+			{
+				Connection->EnableEncryptionWithKey(Response.EncryptionKey);
+			}
+			else
+			{
+				FString ResponseStr(Lex::ToString(Response.Response));
+				UE_LOG(LogBeacon, Warning, TEXT("AOnlineBeaconClient::FinalizeEncryptedConnection: encryption failure [%s] %s"), *ResponseStr, *Response.ErrorMsg);
+				OnFailure();
+			}
+		}
+		else
+		{
+			UE_LOG(LogBeacon, Warning, TEXT("AOnlineBeaconClient::FinalizeEncryptedConnection: connection in invalid state. %s"), *Connection->Describe());
+			OnFailure();
+		}
+	}
+	else
+	{
+		UE_LOG(LogBeacon, Warning, TEXT("AOnlineBeaconClient::FinalizeEncryptedConnection: Connection is null."));
+		OnFailure();
+	}
+}
diff --git a/Engine/Plugins/Online/OnlineSubsystemUtils/Source/OnlineSubsystemUtils/Private/OnlineBeaconHost.cpp b/Engine/Plugins/Online/OnlineSubsystemUtils/Source/OnlineSubsystemUtils/Private/OnlineBeaconHost.cpp
--- a/Engine/Plugins/Online/OnlineSubsystemUtils/Source/OnlineSubsystemUtils/Private/OnlineBeaconHost.cpp
+++ b/Engine/Plugins/Online/OnlineSubsystemUtils/Source/OnlineSubsystemUtils/Private/OnlineBeaconHost.cpp
@@ -119,10 +119,17 @@ void AOnlineBeaconHost::NotifyControlMessage(UNetConnection* Connection, uint8 M
 #if SUPPORTS_ENCRYPTION_TOKEN
 					else
 					{
-						UGameInstance* const Instance = GetGameInstance();
-						if (Instance)
+						if (FNetDelegates::OnReceivedNetworkEncryptionToken.IsBound())
 						{
-							Instance->ReceivedNetworkEncryptionToken(Connection, EncryptionToken, FSimpleDelegate::CreateUObject(this, &AOnlineBeaconHost::SendWelcomeControlMessage, Connection));
+							TWeakObjectPtr<UNetConnection> WeakConnection = Connection;
+							FNetDelegates::OnReceivedNetworkEncryptionToken.Execute(EncryptionToken, FOnEncryptionKeyResponse::CreateUObject(this, &AOnlineBeaconHost::SendWelcomeControlMessage, WeakConnection));
+						}
+						else
+						{
+							FString FailureMsg(TEXT("Encryption failure"));
+							UE_LOG(LogBeacon, Warning, TEXT("%s: No delegate available to handle encryption token, disconnecting."), *Connection->GetName());
+							FNetControlMessage<NMT_Failure>::Send(Connection, FailureMsg);
+							bCloseConnection = true;
 						}
 					}
 #endif
@@ -346,6 +353,39 @@ void AOnlineBeaconHost::SendWelcomeControlMessage(UNetConnection* Connection)
 	}
 }
 
+void AOnlineBeaconHost::SendWelcomeControlMessage(const FEncryptionKeyResponse& Response, TWeakObjectPtr<UNetConnection> WeakConnection)
+{
+	UNetConnection* Connection = WeakConnection.Get();
+	if (Connection)
+	{
+		if (Connection->State != USOCK_Invalid && Connection->State != USOCK_Closed && Connection->Driver)
+		{
+			if (Response.Response == EEncryptionResponse::Success)
+			{
+				Connection->EnableEncryptionWithKeyServer(Response.EncryptionKey);
+				SendWelcomeControlMessage(Connection);
+			}
+			else
+			{
+				FString ResponseStr(Lex::ToString(Response.Response));
+				UE_LOG(LogBeacon, Warning, TEXT("OnlineBeaconHost::SendWelcomeControlMessage: encryption failure [%s] %s"), *ResponseStr, *Response.ErrorMsg);
+				FNetControlMessage<NMT_Failure>::Send(Connection, ResponseStr);
+				Connection->FlushNet();
+				// Can't close the connection here since it will leave the failure message in the send buffer and just close the socket. 
+				// Connection->Close();
+			}
+		}
+		else
+		{
+			UE_LOG(LogBeacon, Warning, TEXT("OnlineBeaconHost::SendWelcomeControlMessage: connection in invalid state. %s"), *Connection->Describe());
+		}
+	}
+	else
+	{
+		UE_LOG(LogBeacon, Warning, TEXT("OnlineBeaconHost::SendWelcomeControlMessage: Connection is null."));
+	}
+}
+
 void AOnlineBeaconHost::RegisterHost(AOnlineBeaconHostObject* NewHostObject)
 {
 	const FString& BeaconType = NewHostObject->GetBeaconType();
diff --git a/Engine/Plugins/Online/OnlineSubsystemUtils/Source/OnlineSubsystemUtils/Public/OnlineBeaconClient.h b/Engine/Plugins/Online/OnlineSubsystemUtils/Source/OnlineSubsystemUtils/Public/OnlineBeaconClient.h
--- a/Engine/Plugins/Online/OnlineSubsystemUtils/Source/OnlineSubsystemUtils/Public/OnlineBeaconClient.h
+++ b/Engine/Plugins/Online/OnlineSubsystemUtils/Source/OnlineSubsystemUtils/Public/OnlineBeaconClient.h
@@ -172,6 +172,15 @@ private:
 	/** Token sent to servers when connecting with an NMT_Hello message. */
 	FString EncryptionToken;
 
+	/**
+	 * Setup the connection for encryption with a given key
+	 * All future packets are expected to be encrypted
+	 *
+	 * @param Response response from the game containing its encryption key or an error message
+	 * @param WeakConnection the connection related to the encryption request
+	 */
+	void FinalizeEncryptedConnection(const FEncryptionKeyResponse& Response, TWeakObjectPtr<UNetConnection> WeakConnection);
+
 	/**
 	 * Called on the server side to open up the actor channel that will allow RPCs to occur
 	 * (DO NOT OVERLOAD, implement OnConnected() instead)
diff --git a/Engine/Plugins/Online/OnlineSubsystemUtils/Source/OnlineSubsystemUtils/Public/OnlineBeaconHost.h b/Engine/Plugins/Online/OnlineSubsystemUtils/Source/OnlineSubsystemUtils/Public/OnlineBeaconHost.h
--- a/Engine/Plugins/Online/OnlineSubsystemUtils/Source/OnlineSubsystemUtils/Public/OnlineBeaconHost.h
+++ b/Engine/Plugins/Online/OnlineSubsystemUtils/Source/OnlineSubsystemUtils/Public/OnlineBeaconHost.h
@@ -11,6 +11,7 @@ class AOnlineBeaconClient;
 class AOnlineBeaconHostObject;
 class FInBunch;
 class UNetConnection;
+struct FEncryptionKeyResponse;
 
 /**
  * Main actor that listens for side channel communication from another Unreal Engine application
@@ -109,6 +110,7 @@ private:
 
 	/** Sends the welcome control message to the client. Used as a delegate if encryption is being negotiated. */
 	void SendWelcomeControlMessage(UNetConnection* Connection);
+	void SendWelcomeControlMessage(const FEncryptionKeyResponse& Response, TWeakObjectPtr<UNetConnection> WeakConnection);
 
 	/** Delegate to route a connection attempt to the appropriate beacon host, by type */
 	DECLARE_DELEGATE_RetVal_OneParam(AOnlineBeaconClient*, FOnBeaconSpawned, UNetConnection*);
diff --git a/Engine/Plugins/Runtime/PacketHandlers/AESHandlerComponent/Source/AESHandlerComponent.Build.cs b/Engine/Plugins/Runtime/PacketHandlers/AESHandlerComponent/Source/AESHandlerComponent.Build.cs
new file mode 100644
--- /dev/null
+++ b/Engine/Plugins/Runtime/PacketHandlers/AESHandlerComponent/Source/AESHandlerComponent.Build.cs
@@ -0,0 +1,51 @@
+// Copyright 1998-2017 Epic Games, Inc. All Rights Reserved.
+
+using UnrealBuildTool;
+using System.IO;
+
+public class AESHandlerComponent : ModuleRules
+{
+    public AESHandlerComponent(ReadOnlyTargetRules Target) : base(Target)
+	{
+		PrivateIncludePaths.AddRange(
+			new string[] {
+				"AESHandlerComponent/Private",
+			}
+			);
+
+		PublicDependencyModuleNames.AddRange(
+			new string[]
+			{
+				"Core",
+				"PacketHandler",
+				"PlatformCrypto",
+			}
+			);
+
+		PublicIncludePathModuleNames.AddRange(
+			new string[]
+			{
+				"PlatformCrypto"
+			}
+			);
+
+		if (Target.Platform == UnrealTargetPlatform.XboxOne)
+		{
+			PublicDependencyModuleNames.AddRange(
+				new string[]
+				{
+					"PlatformCryptoBCrypt",
+				}
+				);
+		}
+		else
+		{
+			PublicDependencyModuleNames.AddRange(
+				new string[]
+				{
+					"PlatformCryptoOpenSSL",
+				}
+				);
+		}
+	}
+}
diff --git a/Engine/Plugins/Runtime/PacketHandlers/AESHandlerComponent/Source/Private/AESHandlerComponent.cpp b/Engine/Plugins/Runtime/PacketHandlers/AESHandlerComponent/Source/Private/AESHandlerComponent.cpp
new file mode 100644
--- /dev/null
+++ b/Engine/Plugins/Runtime/PacketHandlers/AESHandlerComponent/Source/Private/AESHandlerComponent.cpp
@@ -0,0 +1,200 @@
+// Copyright 1998-2017 Epic Games, Inc. All Rights Reserved.
+
+#include "AESHandlerComponent.h"
+
+IMPLEMENT_MODULE( FAESHandlerComponentModule, AESHandlerComponent )
+
+TSharedPtr<HandlerComponent> FAESHandlerComponentModule::CreateComponentInstance(FString& Options)
+{
+	TSharedPtr<HandlerComponent> ReturnVal = NULL;
+
+	ReturnVal = MakeShared<FAESHandlerComponent>();
+
+	return ReturnVal;
+}
+
+
+FAESHandlerComponent::FAESHandlerComponent()
+	: bEncryptionEnabled(false)
+{
+	EncryptionContext = IPlatformCrypto::Get().CreateContext();
+}
+
+void FAESHandlerComponent::SetEncryptionKey(TArrayView<const uint8> NewKey)
+{
+	if (NewKey.Num() != KeySizeInBytes)
+	{
+		UE_LOG(PacketHandlerLog, Log, TEXT("FAESHandlerComponent::SetEncryptionKey. NewKey is not %d bytes long, ignoring."), KeySizeInBytes);
+		return;
+	}
+
+	Key.Reset(KeySizeInBytes);
+	Key.Append(NewKey.GetData(), NewKey.Num());
+}
+
+void FAESHandlerComponent::EnableEncryption()
+{
+	bEncryptionEnabled = true;
+}
+
+void FAESHandlerComponent::DisableEncryption()
+{
+	bEncryptionEnabled = false;
+}
+
+void FAESHandlerComponent::Initialize()
+{
+	SetActive(true);
+	SetState(Handler::Component::State::Initialized);
+	Initialized();
+}
+
+bool FAESHandlerComponent::IsValid() const
+{
+	return true;
+}
+
+void FAESHandlerComponent::Incoming(FBitReader& Packet)
+{
+	DECLARE_SCOPE_CYCLE_COUNTER(TEXT("PacketHandler AES Decrypt"), STAT_PacketHandler_AES_Decrypt, STATGROUP_Net);
+
+	// Handle this packet
+	if (IsValid() && Packet.GetNumBytes() > 0)
+	{
+		// Check first bit to see whether payload is encrypted
+		if (Packet.ReadBit() != 0)
+		{
+			// If the key hasn't been set yet, we can't decrypt, so ignore this packet. We don't set an error in this case because it may just be an out-of-order packet.
+			if (Key.Num() == 0)
+			{
+				UE_LOG(PacketHandlerLog, Log, TEXT("FAESHandlerComponent::Incoming: received encrypted packet before key was set, ignoring."));
+				FBitReader EmptyPacket(nullptr, 0);
+				Packet = EmptyPacket;
+				return;
+			}
+
+			// Copy remaining bits to a TArray so that they are byte-aligned.
+			TArray<uint8> Ciphertext;
+			Ciphertext.AddZeroed(Packet.GetBytesLeft());
+			Packet.SerializeBits(Ciphertext.GetData(), Packet.GetBitsLeft());
+
+			UE_LOG(PacketHandlerLog, VeryVerbose, TEXT("AES packet handler received %ld bytes before decryption."), Ciphertext.Num());
+
+			EPlatformCryptoResult DecryptResult = EPlatformCryptoResult::Failure;
+			TArray<uint8> Plaintext = EncryptionContext->Decrypt_AES_256_ECB(Ciphertext, Key, DecryptResult);
+
+			if (DecryptResult == EPlatformCryptoResult::Failure)
+			{
+				UE_LOG(PacketHandlerLog, Log, TEXT("FAESHandlerComponent::Incoming: failed to decrypt packet."));
+				Packet.SetError();
+				return;
+			}
+			else
+			{
+				// Look for the termination bit that was written in Outgoing() to determine the exact bit size.
+				uint8 LastByte = Plaintext.Last();
+
+				if (LastByte != 0)
+				{
+					int32 BitSize = (Plaintext.Num() * 8) - 1;
+
+					// Bit streaming, starts at the Least Significant Bit, and ends at the MSB.
+					while (!(LastByte & 0x80))
+					{
+						LastByte *= 2;
+						BitSize--;
+					}
+
+					UE_LOG(PacketHandlerLog, VeryVerbose, TEXT("  Have %d bits after decryption."), BitSize);
+
+					FBitReader OutPacket(Plaintext.GetData(), BitSize);
+					Packet = OutPacket;
+				}
+				else
+				{
+					UE_LOG(PacketHandlerLog, Log, TEXT("FAESHandlerComponent::Incoming: malformed packet, last byte was 0."));
+					Packet.SetError();
+				}
+			}
+		}
+	}
+}
+
+void FAESHandlerComponent::Outgoing(FBitWriter& Packet)
+{
+	DECLARE_SCOPE_CYCLE_COUNTER(TEXT("PacketHandler AES Encrypt"), STAT_PacketHandler_AES_Encrypt, STATGROUP_Net);
+
+	// Handle this packet
+	if (IsValid() && Packet.GetNumBytes() > 0)
+	{
+		// Allow for encryption enabled bit and termination bit. Allow resizing to account for encryption padding.
+		FBitWriter NewPacket(Packet.GetNumBits() + 2, true);
+		NewPacket.WriteBit(bEncryptionEnabled ? 1 : 0);
+
+		if (NewPacket.IsError())
+		{
+			UE_LOG(PacketHandlerLog, Log, TEXT("FAESHandlerComponent::Outgoing: failed to write encryption bit."));
+			Packet.SetError();
+			return;
+		}
+
+		if (bEncryptionEnabled)
+		{
+			UE_LOG(PacketHandlerLog, VeryVerbose, TEXT("AES packet handler sending %ld bits before encryption."), Packet.GetNumBits());
+
+			// Write a termination bit so that the receiving side can calculate the exact number of bits sent.
+			// Same technique used in UNetConnection.
+			Packet.WriteBit(1);
+
+			if (Packet.IsError())
+			{
+				UE_LOG(PacketHandlerLog, Log, TEXT("FAESHandlerComponent::Outgoing: failed to write termination bit."));
+				return;
+			}
+
+			EPlatformCryptoResult EncryptResult = EPlatformCryptoResult::Failure;
+			TArray<uint8> Ciphertext = EncryptionContext->Encrypt_AES_256_ECB(TArrayView<uint8>(Packet.GetData(), Packet.GetNumBytes()), Key, EncryptResult);
+
+			if (EncryptResult == EPlatformCryptoResult::Failure)
+			{
+				UE_LOG(PacketHandlerLog, Log, TEXT("FAESHandlerComponent::Outgoing: failed to encrypt packet."));
+				Packet.SetError();
+				return;
+			}
+			else
+			{
+				NewPacket.Serialize(Ciphertext.GetData(), Ciphertext.Num());
+
+				if (NewPacket.IsError())
+				{
+					UE_LOG(PacketHandlerLog, Log, TEXT("FAESHandlerComponent::Outgoing: failed to write ciphertext to packet."));
+					Packet.SetError();
+					return;
+				}
+
+				UE_LOG(PacketHandlerLog, VeryVerbose, TEXT("  AES packet handler sending %d bytes after encryption."), NewPacket.GetNumBytes());
+			}
+		}
+		else
+		{
+			NewPacket.SerializeBits(Packet.GetData(), Packet.GetNumBits());
+		}
+
+		Packet.Reset();
+		Packet.SerializeBits(NewPacket.GetData(), NewPacket.GetNumBits());
+	}
+}
+
+void FAESHandlerComponent::IncomingConnectionless(FString Address, FBitReader& Packet)
+{
+}
+
+void FAESHandlerComponent::OutgoingConnectionless(FString Address, FBitWriter& Packet)
+{
+}
+
+int32 FAESHandlerComponent::GetReservedPacketBits()
+{
+	// Worst case includes the encryption enabled bit, the termination bit, padding up to the next whole byte, and a block of padding.
+	return 2 + 7 + (BlockSizeInBytes * 8);
+}
diff --git a/Engine/Source/Editor/UnrealEd/Classes/Settings/LevelEditorPlaySettings.h b/Engine/Source/Editor/UnrealEd/Classes/Settings/LevelEditorPlaySettings.h
--- a/Engine/Source/Editor/UnrealEd/Classes/Settings/LevelEditorPlaySettings.h
+++ b/Engine/Source/Editor/UnrealEd/Classes/Settings/LevelEditorPlaySettings.h
@@ -297,10 +297,14 @@ private:
 	UPROPERTY(config, EditAnywhere, Category=MultiplayerOptions, meta=(ClampMin = "1", UIMin = "1", UIMax = "64"))
 	int32 PlayNumberOfClients;
 
+	/** What port used by the server for simple networking */
+	UPROPERTY(config, EditAnywhere, Category = MultiplayerOptions)
+	uint16 ServerPort;
+
 	/** Width to use when spawning additional windows. */
 	UPROPERTY(config, EditAnywhere, Category=MultiplayerOptions)
 	int32 ClientWindowWidth;
-
+	
 	/**
 	 * When running multiple players or a dedicated server the client need to connect to the server, this option sets how they connect
 	 *
@@ -360,6 +364,10 @@ public:
 	void SetPlayNumberOfClients( const int32 InPlayNumberOfClients ) { PlayNumberOfClients = InPlayNumberOfClients; }
 	bool IsPlayNumberOfClientsActive() const { return (PlayNetMode != PIE_Standalone) || RunUnderOneProcess; }
 	bool GetPlayNumberOfClients( int32 &OutPlayNumberOfClients ) const { OutPlayNumberOfClients = PlayNumberOfClients; return IsPlayNumberOfClientsActive(); }
+
+	void SetServerPort(const uint16 InServerPort) { ServerPort = InServerPort; }
+	bool IsServerPortActive() const { return (PlayNetMode != PIE_Standalone) || RunUnderOneProcess; }
+	bool GetServerPort(uint16 &OutServerPort) const { OutServerPort = ServerPort; return IsServerPortActive(); }
 	
 	bool IsAutoConnectToServerActive() const { return PlayNumberOfClients > 1 || PlayNetDedicated; }
 	bool GetAutoConnectToServer(bool &OutAutoConnectToServer) const { OutAutoConnectToServer = AutoConnectToServer; return IsAutoConnectToServerActive(); }
diff --git a/Engine/Source/Editor/UnrealEd/Private/PlayLevel.cpp b/Engine/Source/Editor/UnrealEd/Private/PlayLevel.cpp
--- a/Engine/Source/Editor/UnrealEd/Private/PlayLevel.cpp
+++ b/Engine/Source/Editor/UnrealEd/Private/PlayLevel.cpp
@@ -1222,12 +1222,23 @@ void UEditorEngine::StartQueuedPlayMapRequest()
 			NumClients++;
 		}
 
-		// Spawn number of clients
+		// Build the connection String
+		FString ConnectionAddr(TEXT("127.0.0.1"));
 		const bool WillAutoConnectToServer = [&PlayInSettings] { bool AutoConnectToServer(false); return (PlayInSettings->GetAutoConnectToServer(AutoConnectToServer) && AutoConnectToServer); }();
+		if (WillAutoConnectToServer)
+		{
+			uint16 ServerPort = 0;
+			if (PlayInSettings->GetServerPort(ServerPort))
+			{
+				ConnectionAddr += FString::Printf(TEXT(":%hu"), ServerPort);
+			}
+		}
+
+		// Spawn number of clients
 		const int32 PlayNumberOfClients = [&PlayInSettings]{ int32 NumberOfClients(0); return (PlayInSettings->GetPlayNumberOfClients(NumberOfClients) ? NumberOfClients : 0); }();
 		for (int32 i = NumClients; i < PlayNumberOfClients; ++i)
 		{
-			PlayStandaloneLocalPc(TEXT("127.0.0.1"), &WinPosition, i, false);
+			PlayStandaloneLocalPc(*ConnectionAddr, &WinPosition, i, false);
 		}
 	}
 	else
@@ -1358,6 +1369,12 @@ void UEditorEngine::PlayStandaloneLocalPc(FString MapNameOverride, FIntPoint* Wi
 		AdditionalParameters += PlayInSettings->AdditionalLaunchParameters;
 	}
 
+	uint16 ServerPort = 0;
+	if (bIsServer && PlayInSettings->GetServerPort(ServerPort))
+	{
+		AdditionalParameters += FString::Printf(TEXT(" -port=%hu"), ServerPort);
+	}
+
 	// Decide if fullscreen or windowed based on what is specified in the params
 	if (!AdditionalParameters.Contains(TEXT("-fullscreen")) && !AdditionalParameters.Contains(TEXT("-windowed")))
 	{
diff --git a/Engine/Source/Editor/UnrealEd/Private/Settings/LevelEditorPlaySettingsCustomization.h b/Engine/Source/Editor/UnrealEd/Private/Settings/LevelEditorPlaySettingsCustomization.h
--- a/Engine/Source/Editor/UnrealEd/Private/Settings/LevelEditorPlaySettingsCustomization.h
+++ b/Engine/Source/Editor/UnrealEd/Private/Settings/LevelEditorPlaySettingsCustomization.h
@@ -498,6 +498,10 @@ public:
 				.DisplayName(LOCTEXT("NumberOfPlayersLabel", "Number of Players"))
 				.IsEnabled(TAttribute<bool>::Create(TAttribute<bool>::FGetter::CreateSP(this, &FLevelEditorPlaySettingsCustomization::HandlePlayNumberOfClientsIsEnabled)));
 
+			NetworkCategory.AddProperty("ServerPort")
+				.DisplayName(LOCTEXT("ServerPortLabel", "Server Port"))
+				.IsEnabled(TAttribute<bool>::Create(TAttribute<bool>::FGetter::CreateSP(this, &FLevelEditorPlaySettingsCustomization::HandleServerPortIsEnabled)));
+
 			NetworkCategory.AddProperty("AdditionalServerGameOptions")
 				.DisplayName(LOCTEXT("ServerGameOptionsLabel", "Server Game Options"))
 				.IsEnabled(TAttribute<bool>::Create(TAttribute<bool>::FGetter::CreateSP(this, &FLevelEditorPlaySettingsCustomization::HandleGameOptionsIsEnabled)));
@@ -702,6 +706,12 @@ private:
 		return GetDefault<ULevelEditorPlaySettings>()->IsPlayNumberOfClientsActive();
 	}
 
+	// Callback for checking whether the ServerPort is enabled.
+	bool HandleServerPortIsEnabled() const
+	{
+		return GetDefault<ULevelEditorPlaySettings>()->IsServerPortActive();
+	}
+
 	// Callback for checking whether the AdditionalServerGameOptions is enabled.
 	bool HandleGameOptionsIsEnabled( ) const
 	{
diff --git a/Engine/Source/Editor/UnrealEd/Private/Settings/SettingsClasses.cpp b/Engine/Source/Editor/UnrealEd/Private/Settings/SettingsClasses.cpp
--- a/Engine/Source/Editor/UnrealEd/Private/Settings/SettingsClasses.cpp
+++ b/Engine/Source/Editor/UnrealEd/Private/Settings/SettingsClasses.cpp
@@ -370,6 +370,7 @@ ULevelEditorPlaySettings::ULevelEditorPlaySettings( const FObjectInitializer& Ob
 	ClientWindowWidth = 640;
 	ClientWindowHeight = 480;
 	PlayNumberOfClients = 1;
+	ServerPort = 17777;
 	PlayNetDedicated = false;
 	RunUnderOneProcess = true;
 	RouteGamepadToSecondWindow = false;
diff --git a/Engine/Source/Runtime/Core/Private/Misc/Parse.cpp b/Engine/Source/Runtime/Core/Private/Misc/Parse.cpp
--- a/Engine/Source/Runtime/Core/Private/Misc/Parse.cpp
+++ b/Engine/Source/Runtime/Core/Private/Misc/Parse.cpp
@@ -225,65 +225,82 @@ bool FParse::Value(
 	bool			bShouldStopOnSeparator
 )
 {
-	const TCHAR* Found = FCString::Strifind(Stream,Match);
+	bool bSuccess = false;
+	int32 MatchLen = FCString::Strlen(Match);
 
-	if (!Found)
+	for (const TCHAR* Found = FCString::Strifind(Stream, Match); Found != nullptr; Found = FCString::Strifind(Found + MatchLen, Match))
 	{
-		return false;
-	}
+		const TCHAR* Start = Found + MatchLen;
 
-	const TCHAR* Start = Found + FCString::Strlen(Match);
+		// Check for quoted arguments' string with spaces
+		// -Option="Value1 Value2"
+		//         ^~~~Start
+		bool bArgumentsQuoted = *Start == '"';
 
-	// Check for quoted arguments' string with spaces
-	// -Option="Value1 Value2"
-	//         ^~~~Start
-	bool bArgumentsQuoted = *Start == '"';
+		// Number of characters we can look back from found looking for first parenthesis.
+		uint32 AllowedBacktraceCharactersCount = Found - Stream;
 
-	// Number of characters we can look back from found looking for first parenthesis.
-	uint32 AllowedBacktraceCharactersCount = Found - Stream;
+		// Check for fully quoted string with spaces
+		bool bFullyQuoted = 
+			// "-Option=Value1 Value2"
+			//   ^~~~Found
+			AllowedBacktraceCharactersCount > 1 && *(Found - 1) == '-' && *(Found - 2) == '"';
 
-	// Check for fully quoted string with spaces
-	bool bFullyQuoted = 
-		// "Option=Value1 Value2"
-		//  ^~~~Found
-		(AllowedBacktraceCharactersCount > 0 && (*(Found - 1) == '"'))
-		// "-Option=Value1 Value2"
-		//   ^~~~Found
-		|| (AllowedBacktraceCharactersCount > 1 && ((*(Found - 1) == '-') && (*(Found - 2) == '"')));
+		// If we are parsing within a parameter value, this is an invalid match - skip past and try again
+		bool bWithinParamValue = bFullyQuoted &&
+			// -Param="-Option=Value1 Value2"
+			//   ^~~~Found
+			AllowedBacktraceCharactersCount > 2 && *(Found - 3) == '=';
 
-	if (bArgumentsQuoted || bFullyQuoted)
-	{
-		// Skip quote character if only params were quoted.
-		int32 QuoteCharactersToSkip = bArgumentsQuoted ? 1 : 0;
-		FCString::Strncpy(Value, Start + QuoteCharactersToSkip, MaxLen);
+		if (bWithinParamValue)
+		{
+			continue;
+		}
+
+
+		bFullyQuoted = bFullyQuoted ||
+			// "Option=Value1 Value2"
+			//  ^~~~Found
+			(AllowedBacktraceCharactersCount > 0 && *(Found - 1) == '"');
 
-		Value[MaxLen-1]=0;
-		TCHAR* Temp = FCString::Strstr( Value, TEXT("\x22") );
-		if (Temp != nullptr)
+		if (bArgumentsQuoted || bFullyQuoted)
 		{
-			*Temp = 0;
+			// Skip quote character if only params were quoted.
+			int32 QuoteCharactersToSkip = bArgumentsQuoted ? 1 : 0;
+			FCString::Strncpy(Value, Start + QuoteCharactersToSkip, MaxLen);
+
+			Value[MaxLen-1]=0;
+			TCHAR* Temp = FCString::Strstr( Value, TEXT("\x22") );
+			if (Temp != nullptr)
+			{
+				*Temp = 0;
+			}
 		}
-	}
-	else
-	{
-		// Skip initial whitespace
-		Start += FCString::Strspn(Start, TEXT(" \r\n\t"));
-
-		// Non-quoted string without spaces.
-		FCString::Strncpy( Value, Start, MaxLen );
-		Value[MaxLen-1]=0;
-		TCHAR* Temp;
-		Temp = FCString::Strstr( Value, TEXT(" ")  ); if( Temp ) *Temp=0;
-		Temp = FCString::Strstr( Value, TEXT("\r") ); if( Temp ) *Temp=0;
-		Temp = FCString::Strstr( Value, TEXT("\n") ); if( Temp ) *Temp=0;
-		Temp = FCString::Strstr( Value, TEXT("\t") ); if( Temp ) *Temp=0;
-		if (bShouldStopOnSeparator)
+		else
 		{
-			Temp = FCString::Strstr( Value, TEXT(",")  ); if( Temp ) *Temp=0;
-			Temp = FCString::Strstr( Value, TEXT(")")  ); if( Temp ) *Temp=0;
+			// Skip initial whitespace
+			Start += FCString::Strspn(Start, TEXT(" \r\n\t"));
+
+			// Non-quoted string without spaces.
+			FCString::Strncpy( Value, Start, MaxLen );
+			Value[MaxLen-1]=0;
+			TCHAR* Temp;
+			Temp = FCString::Strstr( Value, TEXT(" ")  ); if( Temp ) *Temp=0;
+			Temp = FCString::Strstr( Value, TEXT("\r") ); if( Temp ) *Temp=0;
+			Temp = FCString::Strstr( Value, TEXT("\n") ); if( Temp ) *Temp=0;
+			Temp = FCString::Strstr( Value, TEXT("\t") ); if( Temp ) *Temp=0;
+			if (bShouldStopOnSeparator)
+			{
+				Temp = FCString::Strstr( Value, TEXT(",")  ); if( Temp ) *Temp=0;
+				Temp = FCString::Strstr( Value, TEXT(")")  ); if( Temp ) *Temp=0;
+			}
 		}
+
+		bSuccess = true;
+		break;
 	}
-	return true;
+
+	return bSuccess;
 }
 
 //
diff --git a/Engine/Source/Runtime/Engine/Classes/Engine/ActorChannel.h b/Engine/Source/Runtime/Engine/Classes/Engine/ActorChannel.h
--- a/Engine/Source/Runtime/Engine/Classes/Engine/ActorChannel.h
+++ b/Engine/Source/Runtime/Engine/Classes/Engine/ActorChannel.h
@@ -88,13 +88,20 @@ class ENGINE_API UActorChannel
 
 	TArray< class FOutBunch * >			QueuedExportBunches;			// Bunches that need to be appended to the export list on the next SendBunch call. This list is used for queued RPC's.
 
+#if !UE_BUILD_SHIPPING
+	/** Whether or not to block sending of NMT_ActorChannelFailure (for NetcodeUnitTest) */
+	bool bBlockChannelFailure;
+#endif
+
 	/**
 	 * Default constructor
 	 */
 	UActorChannel(const FObjectInitializer& ObjectInitializer = FObjectInitializer::Get())
 		: UChannel(ObjectInitializer)
+#if !UE_BUILD_SHIPPING
+		, bBlockChannelFailure(false)
+#endif
 	{
-		ChannelClasses[CHTYPE_Actor] = GetClass();
 		ChType = CHTYPE_Actor;
 		bClearRecentActorRefs = true;
 	}
@@ -127,6 +134,8 @@ public:
 	/** Allocate replication tables for the actor channel. */
 	void SetChannelActor( AActor* InActor );
 
+	virtual void NotifyActorChannelOpen(AActor* InActor, FInBunch& InBunch);
+
 	void SetChannelActorForDestroy( struct FActorDestructionInfo *DestructInfo );
 
 	/** Append any export bunches */
diff --git a/Engine/Source/Runtime/Engine/Classes/Engine/Channel.h b/Engine/Source/Runtime/Engine/Classes/Engine/Channel.h
--- a/Engine/Source/Runtime/Engine/Classes/Engine/Channel.h
+++ b/Engine/Source/Runtime/Engine/Classes/Engine/Channel.h
@@ -23,7 +23,10 @@ enum EChannelType
 	CHTYPE_None			= 0,  // Invalid type.
 	CHTYPE_Control		= 1,  // Connection control.
 	CHTYPE_Actor  		= 2,  // Actor-update channel.
+
+	// @todo: Remove and reassign number to CHTYPE_Voice (breaks net compatibility)
 	CHTYPE_File         = 3,  // Binary file transfer.
+
 	CHTYPE_Voice		= 4,  // VoIP data channel
 	CHTYPE_MAX          = 8,  // Maximum.
 };
@@ -65,12 +68,6 @@ class ENGINE_API UChannel
 	class FOutBunch*	OutRec;				// Outgoing reliable unacked data.
 	class FInBunch*		InPartialBunch;		// Partial bunch we are receiving (incoming partial bunches are appended to this)
 
-	/** UChannel statics. */
-	static UClass* ChannelClasses[CHTYPE_MAX];
-	
-	/** @return true if the specified channel type exists. */
-	static bool IsKnownChannelType( int32 Type );
-
 public:
 
 	// UObject overrides
diff --git a/Engine/Source/Runtime/Engine/Classes/Engine/ControlChannel.h b/Engine/Source/Runtime/Engine/Classes/Engine/ControlChannel.h
--- a/Engine/Source/Runtime/Engine/Classes/Engine/ControlChannel.h
+++ b/Engine/Source/Runtime/Engine/Classes/Engine/ControlChannel.h
@@ -82,7 +82,6 @@ class ENGINE_API UControlChannel
 	UControlChannel(const FObjectInitializer& ObjectInitializer = FObjectInitializer::Get())
 		: UChannel(ObjectInitializer)
 	{
-		ChannelClasses[CHTYPE_Control]      = GetClass();
 		ChType								= CHTYPE_Control;
 	}
 
diff --git a/Engine/Source/Runtime/Engine/Classes/Engine/GameInstance.h b/Engine/Source/Runtime/Engine/Classes/Engine/GameInstance.h
--- a/Engine/Source/Runtime/Engine/Classes/Engine/GameInstance.h
+++ b/Engine/Source/Runtime/Engine/Classes/Engine/GameInstance.h
@@ -4,10 +4,11 @@
 
 #include "CoreMinimal.h"
 #include "UObject/ObjectMacros.h"
+#include "UObject/ScriptMacros.h"
 #include "UObject/Object.h"
 #include "Templates/SubclassOf.h"
 #include "Engine/EngineBaseTypes.h"
-#include "UObject/ScriptMacros.h"
+#include "Engine/NetworkDelegates.h"
 #include "GameInstance.generated.h"
 
 class AGameModeBase;
@@ -310,6 +311,12 @@ public:
 	virtual void HandleGameNetControlMessage(class UNetConnection* Connection, uint8 MessageByte, const FString& MessageStr)
 	{}
 	
+	/** Handle setting up encryption keys. Games that override this MUST call the delegate when their own (possibly async) processing is complete. */
+	virtual void ReceivedNetworkEncryptionToken(const FString& EncryptionToken, const FOnEncryptionKeyResponse& Delegate);
+
+	/** Called when a client receives the EncryptionAck control message from the server, will generally enable encryption. */
+	virtual void ReceivedNetworkEncryptionAck(const FOnEncryptionKeyResponse& Delegate);
+
 	/** Call to preload any content before loading a map URL, used during seamless travel as well as map loading */
 	virtual void PreloadContentForURL(FURL InURL);
 
diff --git a/Engine/Source/Runtime/Engine/Classes/Engine/NetConnection.h b/Engine/Source/Runtime/Engine/Classes/Engine/NetConnection.h
--- a/Engine/Source/Runtime/Engine/Classes/Engine/NetConnection.h
+++ b/Engine/Source/Runtime/Engine/Classes/Engine/NetConnection.h
@@ -17,6 +17,7 @@
 #include "Engine/Player.h"
 #include "Engine/Channel.h"
 #include "ProfilingDebugging/Histogram.h"
+#include "ArrayView.h"
 
 #include "NetConnection.generated.h"
 
@@ -119,6 +120,29 @@ namespace EClientLoginState
 	}
 };
 
+
+// Delegates
+#if !UE_BUILD_SHIPPING
+/**
+ * Delegate for hooking the net connections 'ReceivedRawPacket'
+ * 
+ * @param Data				The data received
+ * @param Count				The number of bytes received
+ * @param bBlockReceive		Whether or not to block further processing of the packet (defaults to false)
+*/
+DECLARE_DELEGATE_ThreeParams(FOnReceivedRawPacket, void* /*Data*/, int32 /*Count*/, bool& /*bBlockReceive*/);
+
+/**
+ * Delegate for hooking the net connections 'LowLevelSend' (at the socket level, after PacketHandler parsing)
+ *
+ * @param Data			The data being sent
+ * @param Count			The number of bytes being sent
+ * @param bBlockSend	Whether or not to block the send (defaults to false)
+*/
+DECLARE_DELEGATE_ThreeParams(FOnLowLevelSend, void* /*Data*/, int32 /*Count*/, bool& /*bBlockSend*/);
+#endif
+
+
 #if DO_ENABLE_NET_TEST
 /**
  * An artificially lagged packet
@@ -160,6 +184,10 @@ class UNetConnection : public UPlayer
 	UPROPERTY()
 	class UNetDriver* Driver;	
 
+	/** The class name for the PackageMap to be loaded */
+	UPROPERTY()
+	TSubclassOf<UPackageMap> PackageMapClass;
+
 	UPROPERTY()
 	/** Package map between local and remote. (negotiates net serialization) */
 	class UPackageMap* PackageMap;
@@ -384,6 +412,16 @@ public:
 	 */
 	bool bResendAllDataSinceOpen;
 
+
+#if !UE_BUILD_SHIPPING
+	/** Delegate for hooking ReceivedRawPacket */
+	FOnReceivedRawPacket	ReceivedRawPacketDel;
+
+	/** Delegate for hooking LowLevelSend */
+	FOnLowLevelSend			LowLevelSendDel;
+#endif
+
+
 	/**
 	 * Called to determine if a voice packet should be replicated to this
 	 * connection or any of its child connections
@@ -571,6 +609,15 @@ public:
 	 */
 	ENGINE_API virtual void InitSequence(int32 IncomingSequence, int32 OutgoingSequence);
 
+	/**
+	 * Sets the encryption key and enables encryption.
+	 */
+	ENGINE_API void EnableEncryptionWithKey(TArrayView<const uint8> Key);
+
+	/**
+	 * Sets the encryption key, enables encryption, and sends the encryption ack to the client.
+	 */
+	ENGINE_API void EnableEncryptionWithKeyServer(TArrayView<const uint8> Key);
 
 	/** 
 	* Gets a unique ID for the connection, this ID depends on the underlying connection
diff --git a/Engine/Source/Runtime/Engine/Classes/Engine/NetDriver.h b/Engine/Source/Runtime/Engine/Classes/Engine/NetDriver.h
--- a/Engine/Source/Runtime/Engine/Classes/Engine/NetDriver.h
+++ b/Engine/Source/Runtime/Engine/Classes/Engine/NetDriver.h
@@ -12,6 +12,7 @@
 #include "UObject/CoreNet.h"
 #include "GameFramework/WorldSettings.h"
 #include "PacketHandler.h"
+#include "Channel.h"
 
 #include "NetDriver.generated.h"
 
@@ -28,6 +29,26 @@ class StatelessConnectHandlerComponent;
 class UNetConnection;
 struct FNetworkObjectInfo;
 
+extern ENGINE_API TAutoConsoleVariable<int32> CVarNetAllowEncryption;
+
+// Delegates
+
+#if !UE_BUILD_SHIPPING
+/**
+ * Delegate for hooking ProcessRemoteFunction (used by NetcodeUnitTest)
+ *
+ * @param Actor				The actor the RPC will be called in
+ * @param Function			The RPC to call
+ * @param Parameters		The parameters data blob
+ * @param OutParms			Out parameter information (irrelevant for RPC's)
+ * @param Stack				The script stack
+ * @param SubObject			The sub-object the RPC is being called in (if applicable)
+ * @param bBlockSendRPC		Whether or not to block sending of the RPC (defaults to false)
+ */
+DECLARE_DELEGATE_SevenParams(FOnSendRPC, AActor* /*Actor*/, UFunction* /*Function*/, void* /*Parameters*/,
+									FOutParmRec* /*OutParms*/, FFrame* /*Stack*/, UObject* /*SubObject*/, bool& /*bBlockSendRPC*/);
+#endif
+
 //
 // Whether to support net lag and packet loss testing.
 //
@@ -160,6 +181,7 @@ struct FCompareFActorPriority
 
 struct FActorDestructionInfo
 {
+	TWeakObjectPtr<ULevel>		Level;
 	TWeakObjectPtr<UObject>		ObjOuter;
 	FVector			DestroyedPosition;
 	FNetworkGUID	NetGUID;
@@ -287,6 +309,15 @@ public:
 	UPROPERTY(Config)
 	FName NetDriverName;
 
+	/** The UChannel classes that should be used under this net driver */
+	UClass* ChannelClasses[CHTYPE_MAX];
+	
+	/** @return true if the specified channel type exists. */
+	FORCEINLINE bool IsKnownChannelType(int32 Type)
+	{
+		return Type >= 0 && Type < CHTYPE_MAX && ChannelClasses[Type] != nullptr;
+	}
+
 	/** Change the NetDriver's NetDriverName. This will also reinit packet simulation settings so that settings can be qualified to a specific driver. */
 	void SetNetDriverName(FName NewNetDriverNamed);
 
@@ -425,6 +456,11 @@ public:
 	FDelegateHandle TickFlushDelegateHandle;
 	FDelegateHandle PostTickFlushDelegateHandle;
 
+#if !UE_BUILD_SHIPPING
+	/** Delegate for hooking ProcessRemoteFunction */
+	FOnSendRPC	SendRPCDel;
+#endif
+
 	/** Tracks the amount of time spent during the current frame processing queued bunches. */
 	float ProcessQueuedBunchesCurrentFrameMilliseconds;
 
@@ -773,7 +809,14 @@ protected:
 	int32 ServerReplicateActors_ProcessPrioritizedActors( UNetConnection* Connection, const TArray<FNetViewer>& ConnectionViewers, FActorPriority** PriorityActors, const int32 FinalSortedCount, int32& OutUpdated );
 #endif
 
+	/** Used to handle any NetDriver specific cleanup once a level has been removed from the world. */
+	ENGINE_API virtual void OnLevelRemovedFromWorld(class ULevel* Level, class UWorld* World);
+
+	/** Handle that tracks OnLevelRemovedFromWorld. */
+	FDelegateHandle OnLevelRemovedFromWorldHandle;
+
 private:
+
 	/** Stores the list of objects to replicate into the replay stream. This should be a TUniquePtr, but it appears the generated.cpp file needs the full definition of the pointed-to type. */
 	TSharedPtr<FNetworkObjectList> NetworkObjects;
 
diff --git a/Engine/Source/Runtime/Engine/Classes/Engine/PendingNetGame.h b/Engine/Source/Runtime/Engine/Classes/Engine/PendingNetGame.h
--- a/Engine/Source/Runtime/Engine/Classes/Engine/PendingNetGame.h
+++ b/Engine/Source/Runtime/Engine/Classes/Engine/PendingNetGame.h
@@ -7,6 +7,7 @@
 #include "UObject/UObjectGlobals.h"
 #include "UObject/Object.h"
 #include "Engine/EngineBaseTypes.h"
+#include "NetworkDelegates.h"
 #include "PendingNetGame.generated.h"
 
 class UEngine;
@@ -14,79 +15,6 @@ class UNetConnection;
 class UNetDriver;
 struct FWorldContext;
 
-/**
- * Accepting connection response codes
- */
-namespace EAcceptConnection
-{
-	enum Type
-	{
-		/** Reject the connection */
-		Reject,
-		/** Accept the connection */
-		Accept,
-		/** Ignore the connection, sending no reply, while server traveling */
-		Ignore
-	};
-
-	/** @return the stringified version of the enum passed in */
-	inline const TCHAR* ToString(EAcceptConnection::Type EnumVal)
-	{
-		switch (EnumVal)
-		{
-			case Reject:
-			{
-				return TEXT("Reject");
-			}
-			case Accept:
-			{
-				return TEXT("Accept");
-			}
-			case Ignore:
-			{
-				return TEXT("Ignore");
-			}
-		}
-		return TEXT("");
-	}
-};
-
-/**
- * The net code uses this to send notifications.
- */
-class FNetworkNotify
-{
-public:
-	/**
-	 * Notification that an incoming connection is pending, giving the interface a chance to reject the request
-	 *
-	 * @return EAcceptConnection indicating willingness to accept the connection at this time
-	 */
-	virtual EAcceptConnection::Type NotifyAcceptingConnection() PURE_VIRTUAL(FNetworkNotify::NotifyAcceptedConnection,return EAcceptConnection::Ignore;);
-
-	/**
-	 * Notification that a new connection has been created/established as a result of a remote request, previously approved by NotifyAcceptingConnection
-	 *
-	 * @param Connection newly created connection
-	 */
-	virtual void NotifyAcceptedConnection( class UNetConnection* Connection ) PURE_VIRTUAL(FNetworkNotify::NotifyAcceptedConnection,);
-
-	/**
-	 * Notification that a new channel is being created/opened as a result of a remote request (Actor creation, etc)
-	 *
-	 * @param Channel newly created channel
-	 *
-	 * @return true if the channel should be opened, false if it should be rejected (destroying the channel)
-	 */
-	virtual bool NotifyAcceptingChannel( class UChannel* Channel ) PURE_VIRTUAL(FNetworkNotify::NotifyAcceptingChannel,return false;);
-
-	/** 
-	 * Handler for messages sent through a remote connection's control channel
-	 * not required to handle the message, but if it reads any data from Bunch, it MUST read the ENTIRE data stream for that message (i.e. use FNetControlMessage<TYPE>::Receive())
-	 */
-	virtual void NotifyControlMessage(UNetConnection* Connection, uint8 MessageType, class FInBunch& Bunch) PURE_VIRTUAL(FNetworkNotify::NotifyReceivedText,);
-};
-
 UCLASS(customConstructor, transient)
 class UPendingNetGame :
 	public UObject,
@@ -107,6 +35,15 @@ class UPendingNetGame :
 	 */
 	UPROPERTY()
 	class UDemoNetDriver*	DemoNetDriver;
+	
+	/**
+	 * Setup the connection for encryption with a given key
+	 * All future packets are expected to be encrypted
+	 *
+	 * @param Response response from the game containing its encryption key or an error message
+	 * @param WeakConnection the connection related to the encryption request
+	 */
+	void FinalizeEncryptedConnection(const FEncryptionKeyResponse& Response, TWeakObjectPtr<UNetConnection> WeakConnection);
 
 public:
 	/** URL associated with this level. */
diff --git a/Engine/Source/Runtime/Engine/Classes/Engine/VoiceChannel.h b/Engine/Source/Runtime/Engine/Classes/Engine/VoiceChannel.h
--- a/Engine/Source/Runtime/Engine/Classes/Engine/VoiceChannel.h
+++ b/Engine/Source/Runtime/Engine/Classes/Engine/VoiceChannel.h
@@ -31,8 +31,6 @@ class ENGINE_API UVoiceChannel : public UChannel
 	UVoiceChannel(const FObjectInitializer& ObjectInitializer = FObjectInitializer::Get())
 		: UChannel(ObjectInitializer)
 	{
-		// Register with the network channel system
-		ChannelClasses[CHTYPE_Voice] = GetClass();
 		ChType = CHTYPE_Voice;
 	}
 
diff --git a/Engine/Source/Runtime/Engine/Classes/Engine/World.h b/Engine/Source/Runtime/Engine/Classes/Engine/World.h
--- a/Engine/Source/Runtime/Engine/Classes/Engine/World.h
+++ b/Engine/Source/Runtime/Engine/Classes/Engine/World.h
@@ -56,6 +56,7 @@ class UNetDriver;
 class UPrimitiveComponent;
 class UTexture2D;
 struct FUniqueNetIdRepl;
+struct FEncryptionKeyResponse;
 
 template<typename,typename> class TOctree;
 
@@ -2883,6 +2884,10 @@ private:
 	/** Private version without inlining that does *not* check Dedicated server build flags (which should already have been done). */
 	ENetMode InternalGetNetMode() const;
 
+	// Sends the NMT_Challenge message to Connection.
+	void SendChallengeControlMessage(UNetConnection* Connection);
+	void SendChallengeControlMessage(const FEncryptionKeyResponse& Response, TWeakObjectPtr<UNetConnection> WeakConnection);
+
 public:
 
 #if WITH_EDITOR
diff --git a/Engine/Source/Runtime/Engine/Classes/GameFramework/Actor.h b/Engine/Source/Runtime/Engine/Classes/GameFramework/Actor.h
--- a/Engine/Source/Runtime/Engine/Classes/GameFramework/Actor.h
+++ b/Engine/Source/Runtime/Engine/Classes/GameFramework/Actor.h
@@ -2074,6 +2074,19 @@ public:
 	//--------------------------------------------------------------------------------------
 	// Actor relevancy determination
 
+protected:
+
+	/**
+	 * Determines whether or not the distance between the given SrcLocation and the Actor's location
+	 * is within the net relevancy distance. Actors outside relevancy distance may not be replicated.
+	 *
+	 * @param SrcLocation	Location to test against.
+	 * @return True if the actor is within net relevancy distance, false otherwise.
+	 */
+	bool IsWithinNetRelevancyDistance(const FVector& SrcLocation) const;
+	
+public:	
+
 	/** 
 	  * @param RealViewer - is the "controlling net object" associated with the client for which network relevancy is being checked (typically player controller)
 	  * @param ViewTarget - is the Actor being used as the point of view for the RealViewer
@@ -2321,7 +2334,7 @@ public:
 	bool IsInLevel(const class ULevel *TestLevel) const;
 
 	/** Return the ULevel that this Actor is part of. */
-	ULevel* GetLevel() const { return Cast<ULevel>(GetOuter()); }
+	ULevel* GetLevel() const;
 
 	/**	Do anything needed to clear out cross level references; Called from ULevel::PreSave	 */
 	virtual void ClearCrossLevelReferences();
diff --git a/Engine/Source/Runtime/Engine/Private/Actor.cpp b/Engine/Source/Runtime/Engine/Private/Actor.cpp
--- a/Engine/Source/Runtime/Engine/Private/Actor.cpp
+++ b/Engine/Source/Runtime/Engine/Private/Actor.cpp
@@ -336,6 +336,7 @@ FVector AActor::GetVelocity() const
 
 void AActor::ClearCrossLevelReferences()
 {
+	// TODO: Change GetOutermost to GetLevel?
 	if(RootComponent && GetRootComponent()->GetAttachParent() && (GetOutermost() != GetRootComponent()->GetAttachParent()->GetOutermost()))
 	{
 		GetRootComponent()->DetachFromComponent(FDetachmentTransformRules::KeepRelativeTransform);
@@ -486,7 +487,7 @@ void AActor::RemoveControllingMatineeActor( AMatineeActor& InMatineeActor )
 
 void AActor::BeginDestroy()
 {
-	ULevel* OwnerLevel = Cast<ULevel>(GetOuter());
+	ULevel* OwnerLevel = GetLevel();
 	UnregisterAllComponents();
 	if (OwnerLevel && !OwnerLevel->HasAnyInternalFlags(EInternalObjectFlags::Unreachable))
 	{
@@ -1723,7 +1724,20 @@ bool AActor::ActorHasTag(FName Tag) const
 
 bool AActor::IsInLevel(const ULevel *TestLevel) const
 {
-	return (GetOuter() == TestLevel);
+	return (GetLevel() == TestLevel);
+}
+
+ULevel* AActor::GetLevel() const
+{
+	for (UObject* Outer = GetOuter(); Outer != nullptr; Outer = Outer->GetOuter())
+	{
+		if (ULevel* Level = Cast<ULevel>(Outer))
+		{
+			return Level;
+		}
+	}
+
+	return nullptr;
 }
 
 bool AActor::IsInPersistentLevel(bool bIncludeLevelStreamingPersistent) const
diff --git a/Engine/Source/Runtime/Engine/Private/ActorReplication.cpp b/Engine/Source/Runtime/Engine/Private/ActorReplication.cpp
--- a/Engine/Source/Runtime/Engine/Private/ActorReplication.cpp
+++ b/Engine/Source/Runtime/Engine/Private/ActorReplication.cpp
@@ -235,6 +235,11 @@ void AActor::SyncReplicatedPhysicsSimulation()
 	}
 }
 
+bool AActor::IsWithinNetRelevancyDistance(const FVector& SrcLocation) const
+{
+	return FVector::DistSquared(SrcLocation, GetActorLocation()) < NetCullDistanceSquared;
+}
+
 bool AActor::IsNetRelevantFor(const AActor* RealViewer, const AActor* ViewTarget, const FVector& SrcLocation) const
 {
 	if (bAlwaysRelevant || IsOwnedBy(ViewTarget) || IsOwnedBy(RealViewer) || this == ViewTarget || ViewTarget == Instigator)
@@ -263,12 +268,9 @@ bool AActor::IsNetRelevantFor(const AActor* RealViewer, const AActor* ViewTarget
 		UE_LOG(LogNet, Warning, TEXT("Actor %s / %s has no root component in AActor::IsNetRelevantFor. (Make bAlwaysRelevant=true?)"), *GetClass()->GetName(), *GetName() );
 		return false;
 	}
-	if (GetDefault<AGameNetworkManager>()->bUseDistanceBasedRelevancy)
-	{
-		return ((SrcLocation - GetActorLocation()).SizeSquared() < NetCullDistanceSquared);
-	}
 
-	return true;
+	return !GetDefault<AGameNetworkManager>()->bUseDistanceBasedRelevancy ||
+			IsWithinNetRelevancyDistance(SrcLocation);
 }
 
 bool AActor::IsReplayRelevantFor(const AActor* RealViewer, const AActor* ViewTarget, const FVector& SrcLocation, const float CullDistanceOverrideSq) const
diff --git a/Engine/Source/Runtime/Engine/Private/Character.cpp b/Engine/Source/Runtime/Engine/Private/Character.cpp
--- a/Engine/Source/Runtime/Engine/Private/Character.cpp
+++ b/Engine/Source/Runtime/Engine/Private/Character.cpp
@@ -143,10 +143,6 @@ void ACharacter::PostInitializeComponents()
 void ACharacter::BeginPlay()
 {
 	Super::BeginPlay();
-	if (Mesh)
-	{
-		CacheInitialMeshOffset(Mesh->RelativeLocation, Mesh->RelativeRotation);
-	}
 }
 
 
diff --git a/Engine/Source/Runtime/Engine/Private/DataChannel.cpp b/Engine/Source/Runtime/Engine/Private/DataChannel.cpp
--- a/Engine/Source/Runtime/Engine/Private/DataChannel.cpp
+++ b/Engine/Source/Runtime/Engine/Private/DataChannel.cpp
@@ -1073,13 +1073,6 @@ int32 UChannel::IsNetReady( bool Saturate )
 	return Connection->IsNetReady( Saturate );
 }
 
-
-bool UChannel::IsKnownChannelType( int32 Type )
-{
-	return Type>=0 && Type<CHTYPE_MAX && ChannelClasses[Type];
-}
-
-
 void UChannel::ReceivedNak( int32 NakPacketId )
 {
 	for( FOutBunch* Out=OutRec; Out; Out=Out->Next )
@@ -1105,9 +1098,6 @@ void UChannel::PrintReliableBunchBuffer()
 #endif
 }
 
-
-UClass* UChannel::ChannelClasses[CHTYPE_MAX]={0,0,0,0,0,0,0,0};
-
 /*-----------------------------------------------------------------------------
 	UControlChannel implementation.
 -----------------------------------------------------------------------------*/
@@ -1352,6 +1342,10 @@ void UControlChannel::ReceivedBunch( FInBunch& Bunch )
 					break;
 				case NMT_NetGUIDAssign:
 					FNetControlMessage<NMT_NetGUIDAssign>::Discard(Bunch);
+					break;
+				case NMT_EncryptionAck:
+					//FNetControlMessage<NMT_EncryptionAck>::Discard(Bunch);
+					break;
 				case NMT_BeaconWelcome:
 					//FNetControlMessage<NMT_BeaconWelcome>::Discard(Bunch);
 					break;
@@ -1896,6 +1890,11 @@ void UActorChannel::SetChannelActor( AActor* InActor )
 	Connection->Driver->GetNetworkObjectList().ClearRecentlyDormantConnection(Actor, Connection, Connection->Driver->NetDriverName);
 }
 
+void UActorChannel::NotifyActorChannelOpen(AActor* InActor, FInBunch& InBunch)
+{
+	Actor->OnActorChannelOpen(InBunch, Connection);
+}
+
 void UActorChannel::SetChannelActorForDestroy( FActorDestructionInfo *DestructInfo )
 {
 	check(Connection->Channels[ChIndex]==this);
@@ -2165,7 +2164,12 @@ void UActorChannel::ProcessBunch( FInBunch & Bunch )
 			check( !bSpawnedNewActor );
 			UE_LOG(LogNet, Warning, TEXT("UActorChannel::ProcessBunch: SerializeNewActor failed to find/spawn actor. Actor: %s, Channel: %i"), NewChannelActor ? *NewChannelActor->GetFullName() : TEXT( "NULL" ), ChIndex);
 			Broken = 1;
-			if ( !Connection->InternalAck )
+
+			if (!Connection->InternalAck
+#if !UE_BUILD_SHIPPING
+				&& !bBlockChannelFailure
+#endif
+				)
 			{
 				FNetControlMessage<NMT_ActorChannelFailure>::Send(Connection, ChIndex);
 			}
@@ -2175,7 +2179,7 @@ void UActorChannel::ProcessBunch( FInBunch & Bunch )
 		UE_LOG(LogNetTraffic, Log, TEXT("      Channel Actor %s:"), *NewChannelActor->GetFullName() );
 		SetChannelActor( NewChannelActor );
 
-		Actor->OnActorChannelOpen(Bunch, Connection);
+		NotifyActorChannelOpen(Actor, Bunch);
 
 		RepFlags.bNetInitial = true;
 
@@ -2356,25 +2360,34 @@ bool UActorChannel::ReplicateActor()
 		UE_LOG(LogNet, Log, TEXT( "ReplicateActor: bPausedUntilReliableACK is ending now that reliables have been ACK'd. %s"), *Describe());
 	}
 
-
+	const TArray<FNetViewer>& NetViewers = ActorWorld->GetWorldSettings()->ReplicationViewers;
 	bool bIsNewlyReplicationPaused = false;
 	bool bIsNewlyReplicationUnpaused = false;
-	if (OpenPacketId.First != INDEX_NONE)
+	
+	if (OpenPacketId.First != INDEX_NONE && NetViewers.Num() > 0)
 	{
-		const TArray<FNetViewer>& NetViewers = ActorWorld->GetWorldSettings()->ReplicationViewers;
+		bool bNewPaused = true;
 
-		for (int32 viewerIdx = 0; viewerIdx < NetViewers.Num(); viewerIdx++)
+		for (const FNetViewer& NetViewer : NetViewers)
 		{
-			bool bIsReplicationPausedForConnection = Actor->IsReplicationPausedForConnection(NetViewers[viewerIdx]);
-			bool bOldIsReplicationPaused = IsReplicationPaused();
-			bIsNewlyReplicationPaused = bIsReplicationPausedForConnection && !bOldIsReplicationPaused;
-			bIsNewlyReplicationUnpaused = !bIsReplicationPausedForConnection && bOldIsReplicationPaused;
-			SetReplicationPaused(bIsReplicationPausedForConnection);
-			if (bIsReplicationPausedForConnection && bOldIsReplicationPaused)
+			if (!Actor->IsReplicationPausedForConnection(NetViewer))
 			{
-				return false;
+				bNewPaused = false;
+				break;
 			}
 		}
+
+		const bool bOldPaused = IsReplicationPaused();
+
+		// We were paused and still are, don't do anything.
+		if (bOldPaused && bNewPaused)
+		{
+			return false;
+		}
+
+		bIsNewlyReplicationUnpaused = bOldPaused && !bNewPaused;
+		bIsNewlyReplicationPaused = !bOldPaused && bNewPaused;
+		SetReplicationPaused(bNewPaused);
 	}
 
 	// The package map shouldn't have any carry over guids
diff --git a/Engine/Source/Runtime/Engine/Private/DemoNetDriver.cpp b/Engine/Source/Runtime/Engine/Private/DemoNetDriver.cpp
--- a/Engine/Source/Runtime/Engine/Private/DemoNetDriver.cpp
+++ b/Engine/Source/Runtime/Engine/Private/DemoNetDriver.cpp
@@ -969,14 +969,23 @@ void UDemoNetDriver::TickDispatch( float DeltaSeconds )
 
 void UDemoNetDriver::ProcessRemoteFunction( class AActor* Actor, class UFunction* Function, void* Parameters, struct FOutParmRec* OutParms, struct FFrame* Stack, class UObject* SubObject )
 {
-	if ( IsRecording() )
+#if !UE_BUILD_SHIPPING
+	bool bBlockSendRPC = false;
+
+	SendRPCDel.ExecuteIfBound(Actor, Function, Parameters, OutParms, Stack, SubObject, bBlockSendRPC);
+
+	if (!bBlockSendRPC)
+#endif
 	{
-		if ((Function->FunctionFlags & FUNC_NetMulticast))
+		if ( IsRecording() )
 		{
-			// Handle role swapping if this is a client-recorded replay.
-			FScopedActorRoleSwap RoleSwap(Actor);
+			if ((Function->FunctionFlags & FUNC_NetMulticast))
+			{
+				// Handle role swapping if this is a client-recorded replay.
+				FScopedActorRoleSwap RoleSwap(Actor);
 			
-			InternalProcessRemoteFunction(Actor, SubObject, ClientConnections[0], Function, Parameters, OutParms, Stack, IsServer());
+				InternalProcessRemoteFunction(Actor, SubObject, ClientConnections[0], Function, Parameters, OutParms, Stack, IsServer());
+			}
 		}
 	}
 }
diff --git a/Engine/Source/Runtime/Engine/Private/GameInstance.cpp b/Engine/Source/Runtime/Engine/Private/GameInstance.cpp
--- a/Engine/Source/Runtime/Engine/Private/GameInstance.cpp
+++ b/Engine/Source/Runtime/Engine/Private/GameInstance.cpp
@@ -90,6 +90,9 @@ void UGameInstance::Init()
 				App->RegisterConsoleCommandListener(GenericApplication::FOnConsoleCommandListener::CreateUObject(this, &ThisClass::OnConsoleInput));
 			}
 		}
+
+		FNetDelegates::OnReceivedNetworkEncryptionToken.BindUObject(this, &ThisClass::ReceivedNetworkEncryptionToken);
+		FNetDelegates::OnReceivedNetworkEncryptionAck.BindUObject(this, &ThisClass::ReceivedNetworkEncryptionAck);
 	}
 }
 
@@ -128,6 +131,9 @@ void UGameInstance::Shutdown()
 		}
 	}
 
+	FNetDelegates::OnReceivedNetworkEncryptionToken.Unbind();
+	FNetDelegates::OnReceivedNetworkEncryptionAck.Unbind();
+
 	// Clear the world context pointer to prevent further access.
 	WorldContext = nullptr;
 }
@@ -252,7 +258,15 @@ FGameInstancePIEResult UGameInstance::StartPlayInEditorGameInstance(ULocalPlayer
 	{
 		FString Error;
 		FURL BaseURL = WorldContext->LastURL;
-		if (EditorEngine->Browse(*WorldContext, FURL(&BaseURL, TEXT("127.0.0.1"), (ETravelType)TRAVEL_Absolute), Error) == EBrowseReturnVal::Pending)
+
+		FString URLString(TEXT("127.0.0.1"));
+		uint16 ServerPort = 0;
+		if (PlayInSettings->GetServerPort(ServerPort))
+		{
+			URLString += FString::Printf(TEXT(":%hu"), ServerPort);
+		}
+
+		if (EditorEngine->Browse(*WorldContext, FURL(&BaseURL, *URLString, (ETravelType)TRAVEL_Absolute), Error) == EBrowseReturnVal::Pending)
 		{
 			EditorEngine->TransitionType = TT_WaitingToConnect;
 		}
@@ -327,6 +341,13 @@ FGameInstancePIEResult UGameInstance::StartPlayInEditorGameInstance(ULocalPlayer
 		
 		if (PlayNetMode == PIE_ListenServer)
 		{
+			// Add port
+			uint16 ServerPort = 0;
+			if (PlayInSettings->GetServerPort(ServerPort))
+			{
+				URL.Port = ServerPort;
+			}
+
 			// start listen server with the built URL
 			PlayWorld->Listen(URL);
 		}
@@ -979,6 +1000,18 @@ void UGameInstance::AddUserToReplay(const FString& UserString)
 	}
 }
 
+void UGameInstance::ReceivedNetworkEncryptionToken(const FString& EncryptionToken, const FOnEncryptionKeyResponse& Delegate)
+{
+	FEncryptionKeyResponse Response(EEncryptionResponse::Failure, TEXT("ReceivedNetworkEncryptionToken not implemented"));
+	Delegate.ExecuteIfBound(Response);
+}
+
+void UGameInstance::ReceivedNetworkEncryptionAck(const FOnEncryptionKeyResponse& Delegate)
+{
+	FEncryptionKeyResponse Response(EEncryptionResponse::Failure, TEXT("ReceivedNetworkEncryptionAck not implemented"));
+	Delegate.ExecuteIfBound(Response);
+}
+
 TSubclassOf<UOnlineSession> UGameInstance::GetOnlineSessionClass()
 {
 	return UOnlineSession::StaticClass();
