Copying //UE4/Dev-Networking to //UE4/Dev-Main (Source: //UE4/Dev-Networking @ 3802510)

diff --git a/Engine/Plugins/NetcodeUnitTest/NetcodeUnitTest/Source/NetcodeUnitTest/Public/NUTUtilDebug.h b/Engine/Plugins/NetcodeUnitTest/NetcodeUnitTest/Source/NetcodeUnitTest/Public/NUTUtilDebug.h
--- a/Engine/Plugins/NetcodeUnitTest/NetcodeUnitTest/Source/NetcodeUnitTest/Public/NUTUtilDebug.h
+++ b/Engine/Plugins/NetcodeUnitTest/NetcodeUnitTest/Source/NetcodeUnitTest/Public/NUTUtilDebug.h
@@ -64,6 +64,10 @@ class NETCODEUNITTEST_API FScopedLog
 {
 protected:
 	FScopedLog()
+		: LogCategories()
+		, UnitTest(nullptr)
+		, bRemoteLogging(false)
+		, bSuppressLogging(false)
 	{
 	}
 
@@ -75,13 +79,15 @@ public:
 	 * @param InUnitTest		When tracking netcode-related logs, or doing remote logging, specify the client unit test here
 	 * @param bInRemoteLogging	Whether or not to enable logging on the remote server
 	 */
-	FScopedLog(const TArray<FString>& InLogCategories, UClientUnitTest* InUnitTest=NULL, bool bInRemoteLogging=false)
+	FORCEINLINE FScopedLog(const TArray<FString>& InLogCategories, UClientUnitTest* InUnitTest=nullptr, bool bInRemoteLogging=false)
+		: FScopedLog()
 	{
 		InternalConstruct(InLogCategories, InUnitTest, bInRemoteLogging);
 	}
 
 	// As above, but for a single log category
-	FScopedLog(const FString InLogCategory, UClientUnitTest* InUnitTest=NULL, bool bInRemoteLogging=false)
+	FORCEINLINE FScopedLog(const FString InLogCategory, UClientUnitTest* InUnitTest=nullptr, bool bInRemoteLogging=false)
+		: FScopedLog()
 	{
 		TArray<FString> TempLogCategories;
 		TempLogCategories.Add(InLogCategory);
@@ -105,15 +111,45 @@ protected:
 
 	/** Whether or not this is also controlling remote logging as well */
 	bool bRemoteLogging;
+
+	/** Whether or not to suppress instead of enable logging */
+	bool bSuppressLogging;
 };
 
 /**
- * Version of the above, for scoped logging of all netcode-related logs
+ * Version of FScopedLog, for suppressing instead of enabling log entries
+ */
+class FScopedLogSuppress : public FScopedLog
+{
+public:
+	FORCEINLINE FScopedLogSuppress(const TArray<FString>& InLogCategories, UClientUnitTest* InUnitTest=nullptr,
+									bool bInRemoteLogging=false)
+		: FScopedLog()
+	{
+		bSuppressLogging = true;
+
+		InternalConstruct(InLogCategories, InUnitTest, bInRemoteLogging);
+	}
+
+	FORCEINLINE FScopedLogSuppress(const FString InLogCategory, UClientUnitTest* InUnitTest=nullptr, bool bInRemoteLogging=false)
+		: FScopedLog()
+	{
+		TArray<FString> TempLogCategories;
+		TempLogCategories.Add(InLogCategory);
+
+		bSuppressLogging = true;
+
+		InternalConstruct(TempLogCategories, InUnitTest, bInRemoteLogging);
+	}
+};
+
+/**
+ * Version of FScopedLog, for scoped logging of all netcode-related logs
  */
 class NETCODEUNITTEST_API FScopedLogNet : public FScopedLog
 {
 public:
-	FScopedLogNet(UClientUnitTest* InUnitTest, bool bInRemoteLogging=false)
+	FORCEINLINE FScopedLogNet(UClientUnitTest* InUnitTest, bool bInRemoteLogging=false)
 	{
 		TArray<FString> TempLogCategories;
 
@@ -605,7 +641,16 @@ public:
 				if (CurEntry.LogLine == Data)
 				{
 					bWithinLogTrace = true;
-					GTraceManager->AddTrace(CurEntry.LogLine, false, CurEntry.bDump);
+
+					if (CurEntry.bDump)
+					{
+						GTraceManager->TraceAndDump(CurEntry.LogLine);
+					}
+					else
+					{
+						GTraceManager->AddTrace(CurEntry.LogLine);
+					}
+
 					bWithinLogTrace = false;
 				}
 			}
@@ -618,7 +663,16 @@ public:
 					{
 						// NOTE: Do not use Data for the TraceName, as this makes things much harder to track
 						bWithinLogTrace = true;
-						GTraceManager->AddTrace(CurEntry.LogLine, false, CurEntry.bDump);
+
+						if (CurEntry.bDump)
+						{
+							GTraceManager->TraceAndDump(CurEntry.LogLine);
+						}
+						else
+						{
+							GTraceManager->AddTrace(CurEntry.LogLine);
+						}
+
 						bWithinLogTrace = false;
 					}
 				}
diff --git a/Engine/Plugins/NetcodeUnitTest/NetcodeUnitTest/Source/NetcodeUnitTest/Public/NUTUtilReflection.h b/Engine/Plugins/NetcodeUnitTest/NetcodeUnitTest/Source/NetcodeUnitTest/Public/NUTUtilReflection.h
--- a/Engine/Plugins/NetcodeUnitTest/NetcodeUnitTest/Source/NetcodeUnitTest/Public/NUTUtilReflection.h
+++ b/Engine/Plugins/NetcodeUnitTest/NetcodeUnitTest/Source/NetcodeUnitTest/Public/NUTUtilReflection.h
@@ -132,12 +132,12 @@ class FFuncReflection;
  *	readwriteupcast		- UIntProperty									arraycheck	assign
  *	readwriteupcast		- UInt64Property								arraycheck	assign
  *					- ObjectPropertyBase
- *						- SoftObjectProperty
- *							- SoftClassProperty
+ *	readwrite			- SoftObjectProperty										assign
+ *	inherited				- SoftClassProperty									inherited
  *						- LazyObjectProperty
- *	readwrite			- ObjectProperty								arraycheck
- *	inherited				- ClassProperty								inherited
- *						- WeakObjectProperty
+ *	readwrite			- ObjectProperty								arraycheck	assign
+ *	inherited				- ClassProperty								inherited	inherited
+ *	readwrite			- WeakObjectProperty
  *	readwrite		- StrProperty										arraycheck	assign
  *	readwrite		- StructProperty									arraycheck
  *	readwrite		- TextProperty										arraycheck	assign
@@ -190,6 +190,20 @@ class FFuncReflection;
 //				If you (eventually) add autocomplete to the unit test tool log window, this would be the next logical feature.
 
 
+// @todo #JohnB: Perhaps move up to Core, and add a UCLASS version checking StaticClass (tested before, and it works)
+/**
+ * Determines at compile-time, whether or not a type is defined as a USTRUCT, through TIsUStruct<Type>::Value
+ */
+template<typename T>
+class TIsUStruct
+{
+	template<typename A> static uint8	Check(decltype(&A::StaticStruct));
+	template<typename B> static uint16	Check(...);
+
+public:
+	enum { Value = sizeof(Check<T>(0)) == sizeof(uint8) };
+};
+
 
 /**
  * Used for specifying the warning level, for reflection helpers
@@ -435,6 +449,12 @@ public:
 	 */
 	explicit operator UObject*();
 
+	/**
+	 * Cast to writable FSoftObjectPtr* (where write actions should be performed through FSoftObjectPrr methods, not by pointer assignment).
+	 * For USoftObjectProperty
+	 */
+	explicit operator FSoftObjectPtr*();
+
 
 	/**
 	 * Cast to FScriptArray* pointer (only valid for dynamic arrays), then cast to TArray<Type*>
@@ -458,6 +478,18 @@ public:
 	 */
 	explicit operator void*();
 
+	/**
+	 * Gets a reference to the struct being pointed to, with compile-time and runtime checks on the validity of the struct type
+	 * NOTE: This alters the state of the reflection object.
+	 */
+	template<typename T>
+	FORCEINLINE T* GetStructRef()
+	{
+		static_assert(TIsUStruct<T>::Value, "Only USTRUCT types can be passed to GetStructRef");
+
+		return (T*)(void*)(FVMReflection(*this)[TCHAR_TO_ANSI(*FString::Printf(TEXT("F%s"), *T::StaticStruct()->GetName()))]);
+	}
+
 
 	/**
 	 * Special assignment operators
@@ -586,9 +618,9 @@ private:
 	 */
 	FORCEINLINE bool IsPropertyArray() const
 	{
-		UProperty* CurProp = Cast<UProperty>(FieldInstance);
+		const UProperty* CurProp = Cast<UProperty>(FieldInstance);
 
-		return CurProp != NULL && (CurProp->ArrayDim > 1 || Cast<UArrayProperty>(CurProp) != NULL);
+		return CurProp != nullptr && (CurProp->ArrayDim > 1 || Cast<UArrayProperty>(CurProp) != nullptr);
 	}
 
 	/**
@@ -599,7 +631,7 @@ private:
 	FORCEINLINE bool IsPropertyObject() const
 	{
 		return FieldInstance != nullptr && (FieldInstance->IsA(UObjectProperty::StaticClass()) ||
-				FieldInstance->IsA(UWeakObjectProperty::StaticClass()));
+				FieldInstance->IsA(UWeakObjectProperty::StaticClass()) || FieldInstance->IsA(USoftObjectProperty::StaticClass()));
 	}
 
 	/**
@@ -609,27 +641,32 @@ private:
 	 */
 	FORCEINLINE bool CanCastProperty() const
 	{
-		return !bIsError && BaseAddress != NULL && FieldInstance != NULL && FieldAddress != NULL &&
+		return !bIsError && BaseAddress != nullptr && FieldInstance != nullptr && FieldAddress != nullptr &&
 				FieldInstance->IsA(UProperty::StaticClass()) && (!IsPropertyArray() || (bVerifiedFieldType && bSetArrayElement));
 	}
 
 	/**
 	 * Whether or not the VM reflector points to an object that is ready for casting
 	 *
-	 * @return	Whether or not the VM reflector is ready for casting
+	 * @param PropertyClass		The UObject property to check for
+	 * @return					Whether or not the VM reflector is ready for casting
 	 */
+	template<typename PropertyClass>
 	FORCEINLINE bool CanCastObject() const
 	{
+		static_assert(TPointerIsConvertibleFromTo<PropertyClass, UObjectPropertyBase>::Value,
+						"PropertyClass must derive from UObjectPropertyBase");
+
 		bool bReturnVal = false;
 
-		if (!bIsError && FieldInstance != NULL && (!IsPropertyArray() || (bVerifiedFieldType && bSetArrayElement)))
+		if (!bIsError && FieldInstance != nullptr && (!IsPropertyArray() || (bVerifiedFieldType && bSetArrayElement)))
 		{
-			bool bBaseAddressIsObject = BaseAddress != NULL && FieldInstance->IsA(UClass::StaticClass());
+			bool bBaseAddressIsObject = BaseAddress != nullptr && FieldInstance->IsA(UClass::StaticClass());
 			
 			// This can only happen when the attempt to step-in to an object property (i.e. to assign it to BaseAddress) failed,
-			// due to the object property being NULL. However, NULL is still a valid cast return
-			bool bFieldAddressIsObject = !bBaseAddressIsObject && bNextActionMustBeCast && FieldAddress != NULL &&
-											FieldInstance->IsA(UObjectProperty::StaticClass());
+			// due to the object property being nullptr. However, nullptr is still a valid cast return
+			bool bFieldAddressIsObject = !bBaseAddressIsObject && bNextActionMustBeCast && FieldAddress != nullptr &&
+											FieldInstance->IsA(PropertyClass::StaticClass());
 
 			bReturnVal = bBaseAddressIsObject || bFieldAddressIsObject;
 		}
@@ -759,7 +796,7 @@ private:
 
 
 	/** The current field instance (e.g. a UPROPERTY within AActor or such) */
-	UField* FieldInstance;
+	const UField* FieldInstance;
 
 	/** The address of the current field - OR, if bSetArrayElement - the address of the current array element */
 	void* FieldAddress;
@@ -797,6 +834,9 @@ private:
 };
 
 
+// @todo #JohnB: Why do you use InFuncName below, in cases where you already know the UFunction? Examine/refactor this.
+//					(though sometimes, a nullptr UFunction seems to be possible - perhaps it's for better tracking during errors)
+
 /**
  * FFuncReflection - helper for quickly/concisely setting function parameters through reflection
  */
@@ -811,17 +851,48 @@ private:
 	{
 	}
 
+	FFuncReflection(UFunction* InFunction, const TCHAR* InFuncName, void* InParmsMemory)
+		: FunctionName(InFuncName)
+		, Function(InFunction)
+		, ParmsMemory(Function, (uint8*)InParmsMemory)
+		, ParmsRefl(ParmsMemory)
+	{
+	}
+
 public:
+	/**
+	 * Constructor for creating a function instance and (optionally) filling in its parameters
+	 *
+	 * @param InClassName	The name of the class the function is in
+	 * @param InFuncName	The name of the function
+	 */
 	FFuncReflection(const TCHAR* InClassName, const TCHAR* InFuncName)
 		: FFuncReflection(FindField<UFunction>(FindObject<UClass>(ANY_PACKAGE, InClassName), InFuncName), InFuncName)
 	{
 	}
 
+	/**
+	 * Constructor for creating a function instance and (optionally) filling in its parameters
+	 *
+	 * @param TargetObj		The object the function is in
+	 * @param InFuncName	The name of the function
+	 */
 	FFuncReflection(UObject* TargetObj, const TCHAR* InFuncName)
 		: FFuncReflection((TargetObj != nullptr ? TargetObj->FindFunction(InFuncName) : nullptr), InFuncName)
 	{
 	}
 
+	/**
+	 * Constructor for reading (and optionally modifying) an existing functions parameters (e.g. from received RPC's)
+	 *
+	 * @param InFunction	The function whose parameters we want to read/modify
+	 * @param InParms		The function parameters memory pointer
+	 */
+	FFuncReflection(UFunction* InFunction, void* InParms)
+		: FFuncReflection(InFunction, *InFunction->GetName(), InParms)
+	{
+	}
+
 	FORCEINLINE bool IsValid() const
 	{
 		return Function != nullptr && ParmsMemory.IsValid() && !ParmsRefl.IsError();
diff --git a/Engine/Plugins/NetcodeUnitTest/NetcodeUnitTest/Source/NetcodeUnitTest/Public/Net/NUTUtilNet.h b/Engine/Plugins/NetcodeUnitTest/NetcodeUnitTest/Source/NetcodeUnitTest/Public/Net/NUTUtilNet.h
--- a/Engine/Plugins/NetcodeUnitTest/NetcodeUnitTest/Source/NetcodeUnitTest/Public/Net/NUTUtilNet.h
+++ b/Engine/Plugins/NetcodeUnitTest/NetcodeUnitTest/Source/NetcodeUnitTest/Public/Net/NUTUtilNet.h
@@ -9,6 +9,9 @@
 #include "Engine/World.h"
 #include "Sockets.h"
 
+#include "Net/UnitTestPackageMap.h"
+
+
 class APlayerController;
 class FInBunch;
 class FInternetAddr;
@@ -38,6 +41,7 @@ DECLARE_DELEGATE_FourParams(FOnProcessNetEvent, AActor* /*Actor*/, UFunction* /*
 
 /**
  * Class for encapsulating ProcessEvent and general RPC hooks, implemented globally for each UWorld
+ * NOTE: Presently, all RPC hooks tied to UWorld's only hook client RPC's - global RPC hooks, hook both client and server RPC's
  */
 class FProcessEventHook
 {
@@ -63,6 +67,9 @@ public:
 	 */
 	FProcessEventHook()
 		: NetEventHooks()
+		, EventHooks()
+		, GlobalNetEventHooks()
+		, GlobalEventHooks()
 	{
 	}
 
@@ -85,7 +92,7 @@ public:
 	 * Adds a non-RPC hook for the specified UWorld.
 	 *
 	 * @param InWorld	The UWorld whose RPC's should be hooked
-	 * @param InHook	The delegate for handling the hooked RPC's
+	 * @param InHook	The delegate for handling the hooked events
 	 */
 	void AddEventHook(UWorld* InWorld, FOnProcessNetEvent InHook);
 
@@ -96,8 +103,48 @@ public:
 	 */
 	void RemoveEventHook(UWorld* InWorld);
 
+	/**
+	 * Adds a global RPC hook.
+	 *
+	 * @param InHook	The delegate for handling the hooked RPC's
+	 * @return			Returns a delegate handle, used for removing the hook later
+	 */
+	FDelegateHandle AddGlobalRPCHook(FOnProcessNetEvent InHook);
+
+	/**
+	 * Removes the specified global RPC hook
+	 *
+	 * @param InHandle	The handle to the global RPC hook to remove
+	 */
+	void RemoveGlobalRPCHook(FDelegateHandle InHandle);
+
+	/**
+	 * Adds a global non-RPC hook.
+	 *
+	 * @param InHook	The delegate for handling the hooked events
+	 * @return			Returns a delegate handle, used for removing the hook later
+	 */
+	FDelegateHandle AddGlobalEventHook(FOnProcessNetEvent InHook);
+
+	/**
+	 * Removes the specified global non-RPC hook
+	 *
+	 * @param InHandle	The handle to the global non-RPC hook to remove
+	 */
+	void RemoveGlobalEventHook(FDelegateHandle InHandle);
+
 
 private:
+	/**
+	 * Checks performed prior to adding any hooks
+	 */
+	void PreAddHook();
+
+	/**
+	 * Checks performed after removing any hooks
+	 */
+	void PostRemoveHook();
+
 	/**
 	 * Base hook for AActor::ProcessEventDelegate - responsible for filtering based on RPC events and Actor UWorld
 	 */
@@ -109,6 +156,12 @@ private:
 
 	/** The global list of ProcessEvent hooks, and the UWorld they are associated with */
 	TMap<UWorld*, FOnProcessNetEvent> EventHooks;
+
+	/** The global list of RPC hooks, which aren't associated with a specific UWorld */
+	TArray<FOnProcessNetEvent> GlobalNetEventHooks;
+
+	/** The global list of ProcessEvent hooks, which aren't associated with a specific UWorld */
+	TArray<FOnProcessNetEvent> GlobalEventHooks;
 };
 
 
@@ -193,46 +246,14 @@ public:
 	{
 	}
 
-	void Init()
-	{
-		if (AttachedWorld != NULL)
-		{
-#if TARGET_UE4_CL >= CL_DEPRECATEDEL
-			TickDispatchDelegateHandle  = AttachedWorld->OnTickDispatch().AddRaw(this, &FWorldTickHook::TickDispatch);
-			PostTickFlushDelegateHandle = AttachedWorld->OnPostTickFlush().AddRaw(this, &FWorldTickHook::PostTickFlush);
-#else
-			AttachedWorld->OnTickDispatch().AddRaw(this, &FWorldTickHook::TickDispatch);
-			AttachedWorld->OnPostTickFlush().AddRaw(this, &FWorldTickHook::PostTickFlush);
-#endif
-		}
-	}
-
-	void Cleanup()
-	{
-		if (AttachedWorld != NULL)
-		{
-#if TARGET_UE4_CL >= CL_DEPRECATEDEL
-			AttachedWorld->OnPostTickFlush().Remove(PostTickFlushDelegateHandle);
-			AttachedWorld->OnTickDispatch().Remove(TickDispatchDelegateHandle);
-#else
-			AttachedWorld->OnPostTickFlush().RemoveRaw(this, &FWorldTickHook::PostTickFlush);
-			AttachedWorld->OnTickDispatch().RemoveRaw(this, &FWorldTickHook::TickDispatch);
-#endif
-		}
+	void Init();
 
-		AttachedWorld = NULL;
-	}
+	void Cleanup();
 
 private:
-	void TickDispatch(float DeltaTime)
-	{
-		GActiveLogWorld = AttachedWorld;
-	}
+	void TickDispatch(float DeltaTime);
 
-	void PostTickFlush()
-	{
-		GActiveLogWorld = NULL;
-	}
+	void PostTickFlush();
 
 public:
 	/** The world this is attached to */
@@ -249,6 +270,7 @@ private:
 };
 
 
+// @todo #JohnB: Refactor to be based on minimal client
 /**
  * Hooks netcode object serialization, in order to replace replication of a specific object, with another specified object,
  * for the lifetime of the scoped instance
@@ -267,6 +289,26 @@ private:
 	UObject* ObjToReplace;
 };
 
+/**
+ * Hooks netcode name serialization, in order to replace replication of a specific name, with another specified name,
+ * for the lifetime of the scoped instance
+ */
+class NETCODEUNITTEST_API FScopedNetNameReplace
+{
+public:
+	FScopedNetNameReplace(UMinimalClient* InMinClient, const FOnSerializeName::FDelegate& InDelegate);
+
+	// @todo #JohnB: This variation is not tested
+	FScopedNetNameReplace(UMinimalClient* InMinClient, FName InNameToReplace, FName InNameReplacement);
+
+	~FScopedNetNameReplace();
+
+private:
+	UMinimalClient* MinClient;
+
+	FDelegateHandle Handle;
+};
+
 
 /**
  * Netcode based utility functions
@@ -281,11 +323,11 @@ struct NUTNet
 	 * @param InBeacon		The beacon that should be setup
 	 * @param InConnection	The connection associated with the beacon
 	 */
-	static void HandleBeaconReplicate(AActor* InBeacon, UNetConnection* InConnection);
+	static NETCODEUNITTEST_API void HandleBeaconReplicate(AActor* InBeacon, UNetConnection* InConnection);
 
 
 	/**
-	 * Creates a barebones/minimal UWorld, for setting up minimal fake player connections,
+	 * Creates a barebones/minimal UWorld, for setting up minimal client connections,
 	 * and as a container for objects in the unit test commandlet
 	 *
 	 * @return	Returns the created UWorld object
diff --git a/Engine/Plugins/NetcodeUnitTest/NetcodeUnitTest/Source/NetcodeUnitTest/Public/NetcodeUnitTest.h b/Engine/Plugins/NetcodeUnitTest/NetcodeUnitTest/Source/NetcodeUnitTest/Public/NetcodeUnitTest.h
--- a/Engine/Plugins/NetcodeUnitTest/NetcodeUnitTest/Source/NetcodeUnitTest/Public/NetcodeUnitTest.h
+++ b/Engine/Plugins/NetcodeUnitTest/NetcodeUnitTest/Source/NetcodeUnitTest/Public/NetcodeUnitTest.h
@@ -3,12 +3,13 @@
 #pragma once
 
 #include "CoreMinimal.h"
-#include "Misc/EnumClassFlags.h"
 
 /**
  * Compatibility defines
  */
 
+// @todo #JohnB: Consider discarding the compatibility defines. Or at least, removing the ones made redundant by refactoring.
+
 // @todo #JohnBRefactor: Perhaps move compatibility defines to its own header
 
 // @todo #JohnB: I need accurate mainline CL numbers for UT integrations (they aren't included in the merge changelist descriptions)
@@ -36,6 +37,7 @@
  * List of mainline UE4 engine changelists, that required a NetcodeUnitTest compatibility adjustment (newest first)
  * Every time you update NetcodeUnitTest code, to toggle based on a changelist, add a define for it with the CL number here
  */
+#define CL_HELLOENCRYPTION	3578855
 #define CL_PREVERSIONCHANGE	2960134	// Not an accurate CL, just refers to the first version of UE4 code I can be sure changes work on
 #define CL_STATELESSCONNECT	2866629	// @todo #JohnB: Refers to Dev-Networking, update so it refers to /UE4//Main merge CL
 #define CL_FENGINEVERSION	2655102
@@ -60,86 +62,15 @@
  *
  * If in doubt, set to the top CL from list above ('List of mainline UE4 engine changelists'), and work your way down until it compiles.
  */
-#define TARGET_UE4_CL 2960134	// IMPORTANT: If you're hovering over this because compiling failed, you need to adjust this value.
+#define TARGET_UE4_CL 3645298	// IMPORTANT: If you're hovering over this because compiling failed, you need to adjust this value.
 
-class Error;
-class UNetConnection;
-class UUnitTest;
-class UUnitTestBase;
-class UUnitTestManager;
 
 /**
  * Forward declarations
  */
-class UNetConnection;
-class UUnitTestManager;
-class UUnitTestBase;
-
-
-// @todo #JohnBLowPri: Status window logs (such as unit test success/failure), don't appear to add the same colouring to main game log window;
-//				add this sometime (probably based on how FUnitTestProcess does it)
-
-/**
- * Enums
- */
-
-// @todo #JohnBRefactor: Trim values from the below enum - not all are necessary
-
-/**
- * Used to help identify what type of log is being processed
- *
- * NOTE: This enum is not compatible with UENUM (that requires uint8)
- */
-enum class ELogType : uint32
-{
-	None				= 0x00000000,					// Not set
-
-	/** What part of the engine the log message originates from locally (may not be set) */
-	OriginUnitTest		= 0x00000001,					// Log originating from unit test code
-	OriginEngine		= 0x00000002,					// Log originating from an engine event in the unit test (e.g. unit test Tick)
-	OriginNet			= 0x00000004,					// Log originating from netcode for current unit test (specifically net receive)
-	OriginConsole		= 0x00000008,					// Log originating from a console command (typically from the unit test window)
-	OriginVoid			= 0x00000010,					// Log which should have no origin, and which should be ignored by log capturing
-
-	OriginMask			= OriginUnitTest | OriginEngine | OriginNet | OriginConsole | OriginVoid,
-
-
-	/** What class of unit test log this is */
-	Local				= 0x00000080,					// Log from locally executed code (displayed in 'Local' tab)
-	Server				= 0x00000100,					// Log from a server instance (displayed in 'Server' tab)
-	Client				= 0x00000200,					// Log from a client instance (displayed in 'Client' tab)
-
-
-	/** What class of unit test status log this is (UNIT_LOG vs UNIT_STATUS_LOG/STATUS_LOG) */
-	GlobalStatus		= 0x00000400,					// Status placed within the overall status window
-	UnitStatus			= 0x00000800,					// Status placed within the unit test window
-
-	/** Status log modifiers (used with the UNIT_LOG and UNIT_STATUS_LOG/STATUS_LOG macro's) */
-	StatusImportant		= 0x00001000,					// An important status event (displayed in the 'Summary' tab)
-	StatusSuccess		= 0x00002000 | StatusImportant,	// Success event status
-	StatusWarning		= 0x00004000 | StatusImportant,	// Warning event status
-	StatusFailure		= 0x00008000 | StatusImportant,	// Failure event status
-	StatusError			= 0x00010000 | StatusFailure,	// Error/Failure event status, that triggers an overall unit test failure
-	StatusDebug			= 0x00020000,					// Debug status (displayed in the 'Debug' tab)
-	StatusAdvanced		= 0x00040000,					// Status event containing advanced/technical information
-	StatusVerbose		= 0x00080000,					// Status event containing verbose information
-	StatusAutomation	= 0x00100000,					// Status event which should be printed out to the automation tool
-
 
-	/** Log output style modifiers */
-	StyleBold			= 0x00200000,					// Output text in bold
-	StyleItalic			= 0x00400000,					// Output text in italic
-	StyleUnderline		= 0x00800000,					// Output pseudo-underline text (add newline and --- chars; UE4 can't underline)
-	StyleMonospace		= 0x01000000,					// Output monospaced text (e.g. for list tab formatting); can't use bold/italic
-
-	All					= 0xFFFFFFFF,
-
-
-	/** Log lines that should request focus when logged (i.e. if not displayed in the current tab, switch to a tab that does display) */
-	FocusMask			= OriginConsole,
-};
+class UUnitTestManager;
 
-ENUM_CLASS_FLAGS(ELogType);
 
 
 /**
@@ -149,23 +80,9 @@ ENUM_CLASS_FLAGS(ELogType);
 /** Holds a reference to the object in charge of managing unit tests */
 extern UUnitTestManager* GUnitTestManager;
 
-
-// IMPORTANT: If you add more engine-log-capture globals, you must add them to the 'UNIT_EVENT_CLEAR' and related macros
-
-/** Used by to aid in hooking log messages triggered by unit tests */
-extern NETCODEUNITTEST_API UUnitTestBase* GActiveLogUnitTest;
-extern UUnitTestBase* GActiveLogEngineEvent;
-extern UWorld* GActiveLogWorld;				// @todo #JohnBReview: Does this make other log hooks redundant?
-
-/** Keeps track of the UUnitTestNetConnection, currently processing received data (to aid with selective log hooking) */
-extern UNetConnection* GActiveReceiveUnitConnection;
-
 /** Whether not an actor channel, is in the process of initializing the remote actor */
 extern bool GIsInitializingActorChan;
 
-/** The current ELogType flag modifiers, associated with a UNIT_LOG or UNIT_STATUS_LOG/STATUS_LOG call */
-extern NETCODEUNITTEST_API ELogType GActiveLogTypeFlags;
-
 
 /**
  * Declarations
@@ -192,14 +109,6 @@ DECLARE_LOG_CATEGORY_EXTERN(NetCodeTestNone, Log, All);
  * Macro defines
  */
 
-/**
- * Helper function, for macro defines - allows the 'UnitLogTypeFlags' macro parameters below, to be optional
- */
-inline ELogType OptionalFlags(ELogType InFlags=ELogType::None)
-{
-	return InFlags;
-}
-
 // Actual asserts (not optimized out, like almost all other engine assert macros)
 // @todo #JohnBLowPri: Try to get this to show up a message box, or some other notification, rather than just exiting immediately
 #define UNIT_ASSERT(Condition) \
@@ -214,178 +123,3 @@ inline ELogType OptionalFlags(ELogType InFlags=ELogType::None)
 	}
 
 
-// Used with the below macros, and sometimes to aid with large unit tests logs
-
-#define UNIT_LOG_BEGIN(UnitTestObj, UnitLogTypeFlags) \
-	{ \
-		UUnitTestBase* OldUnitLogVal = GActiveLogUnitTest; \
-		GActiveLogUnitTest = (UnitTestObj != nullptr ? UnitTestObj : GActiveLogUnitTest); \
-		GActiveLogTypeFlags = ELogType::UnitStatus | OptionalFlags(UnitLogTypeFlags);
-
-#define UNIT_LOG_END() \
-		GActiveLogTypeFlags = ELogType::None; \
-		GActiveLogUnitTest = OldUnitLogVal; \
-	}
-
-#define UNIT_LOG_VOID_START() \
-	GActiveLogTypeFlags = ELogType::OriginVoid;
-
-#define UNIT_LOG_VOID_END() \
-	GActiveLogTypeFlags = ELogType::None;
-
-
-/**
- * Special log macro for unit tests, to aid in hooking logs from these unit tests
- * NOTE: These logs are scoped
- *
- * @param UnitTestObj		Reference to the unit test this log message is originating from
- * @param UnitLogTypeFlags	(Optional) The additional ELogFlags flags, to be applied to this log
- * @param Format			The format of the log message
- */
-#define UNIT_LOG_OBJ(UnitTestObj, UnitLogTypeFlags, Format, ...) \
-	UNIT_LOG_BEGIN(UnitTestObj, UnitLogTypeFlags); \
-		if ((OptionalFlags(UnitLogTypeFlags) & ELogType::StatusError) == ELogType::StatusError) \
-		{ \
-			UE_LOG(LogUnitTest, Error, Format, ##__VA_ARGS__); \
-		} \
-		else if ((OptionalFlags(UnitLogTypeFlags) & ELogType::StatusWarning) == ELogType::StatusWarning) \
-		{ \
-			UE_LOG(LogUnitTest, Warning, Format, ##__VA_ARGS__); \
-		} \
-		else \
-		{ \
-			UE_LOG(LogUnitTest, Log, Format, ##__VA_ARGS__); \
-		} \
-	UNIT_LOG_END(); \
-	UNIT_LOG_VOID_START(); \
-	if ((OptionalFlags(UnitLogTypeFlags) & ELogType::StatusError) == ELogType::StatusError) \
-	{ \
-		UE_LOG(LogUnitTest, Error, TEXT("%s: %s"), (UnitTestObj != nullptr ? *(UnitTestObj->GetUnitTestName()) : TEXT("nullptr")), \
-				*FString::Printf(Format, ##__VA_ARGS__)); \
-	} \
-	else if ((OptionalFlags(UnitLogTypeFlags) & ELogType::StatusWarning) == ELogType::StatusWarning) \
-	{ \
-		UE_LOG(LogUnitTest, Warning, TEXT("%s: %s"), (UnitTestObj != nullptr ? *(UnitTestObj->GetUnitTestName()) : TEXT("nullptr")), \
-				*FString::Printf(Format, ##__VA_ARGS__)); \
-	} \
-	UNIT_LOG_VOID_END();
-
-
-// More-concise version of the above macro
-
-#define UNIT_LOG(UnitLogTypeFlags, Format, ...) \
-	UNIT_LOG_OBJ(this, UnitLogTypeFlags, Format, ##__VA_ARGS__)
-
-
-// @todo #JohnBDoc: Document these macros
-
-#define UNIT_EVENT_BEGIN(UnitTestObj) \
-	UUnitTestBase* OldEventLogVal = GActiveLogEngineEvent; \
-	GActiveLogEngineEvent = UnitTestObj;
-	
-#define UNIT_EVENT_END \
-	GActiveLogEngineEvent = OldEventLogVal;
-
-/**
- * Stores and then clears all engine-log-capture events (in order to prevent capture of a new log entry)
- * NOTE: Do not clear GActiveLogUnitTest here, as some macros that use this, rely upon that
- */
-#define UNIT_EVENT_CLEAR \
-	UUnitTestBase* OldEventLogVal = GActiveLogEngineEvent; \
-	UWorld* OldEventLogWorld = GActiveLogWorld; \
-	GActiveLogEngineEvent = nullptr; \
-	GActiveLogWorld = nullptr;
-
-/**
- * Restores all stored/cleared engine-log-capture events
- */
-#define UNIT_EVENT_RESTORE \
-	UNIT_EVENT_END \
-	GActiveLogWorld = OldEventLogWorld;
-
-
-/**
- * Special log macro, for messages that should be printed to the unit test status window
- */
-#define STATUS_LOG_BASE(UnitLogTypeFlags, Format, ...) \
-	{ \
-		GActiveLogTypeFlags = ELogType::GlobalStatus | OptionalFlags(UnitLogTypeFlags); \
-		UUnitTestManager* Manager = UUnitTestManager::Get(); \
-		Manager->SetStatusLog(true); \
-		Manager->Logf(Format, ##__VA_ARGS__); \
-		Manager->SetStatusLog(false); \
-		\
-		if ((GActiveLogTypeFlags & ELogType::StatusAutomation) == ELogType::StatusAutomation && GIsAutomationTesting) \
-		{ \
-			GWarn->Logf(Format, ##__VA_ARGS__); \
-		} \
-		\
-		GActiveLogTypeFlags = ELogType::None; \
-	}
-
-#define STATUS_LOG_OBJ(InUnitTest, UnitLogTypeFlags, Format, ...) \
-	{ \
-		GActiveLogTypeFlags = ELogType::GlobalStatus | OptionalFlags(UnitLogTypeFlags); \
-		UUnitTest* SourceUnitTest = InUnitTest; \
-		UNIT_EVENT_CLEAR; \
-		if ((GActiveLogTypeFlags & ELogType::StatusError) == ELogType::StatusError) \
-		{ \
-			if (SourceUnitTest != nullptr) \
-			{ \
-				UE_LOG(LogUnitTest, Error, TEXT("%s: %s"), *SourceUnitTest->GetUnitTestName(), *FString::Printf(Format, ##__VA_ARGS__)); \
-			} \
-			else \
-			{ \
-				UE_LOG(LogUnitTest, Error, Format, ##__VA_ARGS__) \
-			} \
-		} \
-		else if ((GActiveLogTypeFlags & ELogType::StatusWarning) == ELogType::StatusWarning) \
-		{ \
-			if (SourceUnitTest != nullptr) \
-			{ \
-				UE_LOG(LogUnitTest, Warning, TEXT("%s: %s"), *SourceUnitTest->GetUnitTestName(), \
-						*FString::Printf(Format, ##__VA_ARGS__)); \
-			} \
-			else \
-			{ \
-				UE_LOG(LogUnitTest, Warning, Format, ##__VA_ARGS__) \
-			} \
-		} \
-		else \
-		{ \
-			UE_LOG(LogUnitTest, Log, Format, ##__VA_ARGS__) \
-		} \
-		UNIT_EVENT_RESTORE; \
-		STATUS_LOG_BASE(UnitLogTypeFlags, Format, ##__VA_ARGS__) \
-		GActiveLogTypeFlags = ELogType::None; \
-	}
-
-#define STATUS_LOG(UnitLogTypeFlags, Format, ...) STATUS_LOG_OBJ(nullptr, UnitLogTypeFlags, Format, ##__VA_ARGS__)
-
-/**
- * Version of the above, for unit test status window entries, which are from specific unit tests
- */
-#define UNIT_STATUS_LOG_OBJ(UnitTestObj, UnitLogTypeFlags, Format, ...) \
-	UNIT_LOG_BEGIN(UnitTestObj, UnitLogTypeFlags); \
-	STATUS_LOG_OBJ(UnitTestObj, UnitLogTypeFlags, Format, ##__VA_ARGS__) \
-	UNIT_LOG_END();
-
-#define UNIT_STATUS_LOG(UnitLogTypeFlags, Format, ...) \
-	UNIT_STATUS_LOG_OBJ(this, UnitLogTypeFlags, Format, ##__VA_ARGS__)
-
-
-/**
- * Changes/resets the colour of messages printed to the unit test status window
- */
-#define STATUS_SET_COLOR(InColor) \
-	{ \
-		UUnitTestManager* Manager = UUnitTestManager::Get(); \
-		Manager->SetStatusColor(InColor); \
-	}
-
-#define STATUS_RESET_COLOR() \
-	{ \
-		UUnitTestManager* Manager = UUnitTestManager::Get(); \
-		Manager->SetStatusColor(); /* No value specified = reset */ \
-	}
-
diff --git a/Engine/Plugins/NetcodeUnitTest/NetcodeUnitTest/Source/NetcodeUnitTest/Public/UI/LogWindowManager.h b/Engine/Plugins/NetcodeUnitTest/NetcodeUnitTest/Source/NetcodeUnitTest/Public/UI/LogWindowManager.h
--- a/Engine/Plugins/NetcodeUnitTest/NetcodeUnitTest/Source/NetcodeUnitTest/Public/UI/LogWindowManager.h
+++ b/Engine/Plugins/NetcodeUnitTest/NetcodeUnitTest/Source/NetcodeUnitTest/Public/UI/LogWindowManager.h
@@ -3,8 +3,12 @@
 #pragma once
 
 #include "CoreMinimal.h"
+
 #include "NetcodeUnitTest.h"
+#include "UnitLogging.h"
+
 
+// Forward declarations
 class SLogWindow;
 class SWindow;
 
diff --git a/Engine/Plugins/NetcodeUnitTest/NetcodeUnitTest/Source/NetcodeUnitTest/Public/UI/SLogWidget.h b/Engine/Plugins/NetcodeUnitTest/NetcodeUnitTest/Source/NetcodeUnitTest/Public/UI/SLogWidget.h
--- a/Engine/Plugins/NetcodeUnitTest/NetcodeUnitTest/Source/NetcodeUnitTest/Public/UI/SLogWidget.h
+++ b/Engine/Plugins/NetcodeUnitTest/NetcodeUnitTest/Source/NetcodeUnitTest/Public/UI/SLogWidget.h
@@ -10,7 +10,10 @@
 #include "Widgets/SCompoundWidget.h"
 #include "Widgets/Input/SComboBox.h"
 #include "Framework/Docking/TabManager.h"
+
 #include "NetcodeUnitTest.h"
+#include "UnitLogging.h"
+
 
 class FUICommandList;
 class SEditableTextBox;
@@ -56,11 +59,15 @@ struct FLogLine
 	/** The color of the log line */
 	FSlateColor			LogColor;
 
+	/** The text within the log line, to highlight (from searches) */
+	FString				LogHighlight;
+
 
 	FLogLine(ELogType InLogType, TSharedRef<FString> InLogLine, FSlateColor InLogColor)
 		: LogType(InLogType)
 		, LogLine(InLogLine)
 		, LogColor(InLogColor)
+		, LogHighlight()
 	{
 	}
 };
@@ -99,9 +106,21 @@ struct FLogTabInfo
 	/** The text box for the 'Find' bar */
 	TSharedPtr<SEditableTextBox> FindBox;
 
+	/** The label that's displayed when there's an error searching, or searching past the start/end of the list */
+	TSharedPtr<SBorder> FindErrorLabel;
+
+	/** The text block for the above label */
+	TSharedPtr<STextBlock> FindErrorText;
+
+	/** Caches the text of the last search */
+	FString LastFind;
+
 	/** Whether or not the last search direction was upwards */
 	bool bLastFindWasUp;
 
+	/** Whether or not to highlight text, from searches */
+	bool bHighlightFind;
+
 
 	/** Automatically generated name given to the tab */
 	FName TabIdName;
@@ -125,7 +144,11 @@ struct FLogTabInfo
 		, LogListView()
 		, FindWidgets()
 		, FindBox()
+		, FindErrorLabel()
+		, FindErrorText()
+		, LastFind()
 		, bLastFindWasUp(false)
+		, bHighlightFind(true)
 		, TabWidget()
 	{
 		static uint8 TabNameCounter = 0;
@@ -276,6 +299,19 @@ protected:
 	 */
 	bool CanAutoScroll(TSharedPtr<FLogTabInfo> InTab);
 
+	/**
+	 * Updates the find text highlighting for a tab
+	 *
+	 * @param InTab				The tab to update
+	 * @param bHighlight		Whether or not highlighting has been enabled (NOTE: This often differs from the tabs bHighlightFind value)
+	 * @param HighlightText		The text to highlight, if any
+	 */
+	void UpdateFindHighlight(TSharedPtr<FLogTabInfo> InTab, bool bHighlight, FString HighlightText=TEXT(""));
+
+
+	virtual void Tick(const FGeometry& AllottedGeometry, const double InCurrentTime, const float InDeltaTime) override;
+
+
 protected:
 	/** The log window tab manager */
 	TSharedPtr<FTabManager> LogTabManager;
diff --git a/Engine/Plugins/NetcodeUnitTest/NetcodeUnitTest/Source/NetcodeUnitTest/Public/UI/SLogWindow.h b/Engine/Plugins/NetcodeUnitTest/NetcodeUnitTest/Source/NetcodeUnitTest/Public/UI/SLogWindow.h
--- a/Engine/Plugins/NetcodeUnitTest/NetcodeUnitTest/Source/NetcodeUnitTest/Public/UI/SLogWindow.h
+++ b/Engine/Plugins/NetcodeUnitTest/NetcodeUnitTest/Source/NetcodeUnitTest/Public/UI/SLogWindow.h
@@ -3,11 +3,14 @@
 #pragma once
 
 #include "CoreMinimal.h"
-#include "NetcodeUnitTest.h"
 #include "Widgets/DeclarativeSyntaxSupport.h"
 #include "Input/Reply.h"
 #include "Widgets/SWindow.h"
 
+#include "NetcodeUnitTest.h"
+#include "UnitLogging.h"
+
+
 class SLogWidget;
 
 /**
diff --git a/Engine/Plugins/NetcodeUnitTest/NetcodeUnitTest/Source/NetcodeUnitTest/Public/UnitTestEnvironment.h b/Engine/Plugins/NetcodeUnitTest/NetcodeUnitTest/Source/NetcodeUnitTest/Public/UnitTestEnvironment.h
--- a/Engine/Plugins/NetcodeUnitTest/NetcodeUnitTest/Source/NetcodeUnitTest/Public/UnitTestEnvironment.h
+++ b/Engine/Plugins/NetcodeUnitTest/NetcodeUnitTest/Source/NetcodeUnitTest/Public/UnitTestEnvironment.h
@@ -126,7 +126,7 @@ public:
 	}
 
 	/**
-	 * Returns the default URL that fake clients should use, when connecting to a server, for the current game
+	 * Returns the default URL that minimal clients should use, when connecting to a server, for the current game
 	 *
 	 * @return	Returns the default URL
 	 */
@@ -200,6 +200,14 @@ public:
 	{
 	}
 
+	/**
+	 * If the game environment requires special/complex setup (e.g. player auth before join),
+	 * then UnitTask's for handling complex setup, are added here.
+	 */
+	virtual void InitializeUnitTasks()
+	{
+	}
+
 
 protected:
 	// @todo #JohnB: Replace references to UnitTestFlags in above functions, with usage of the UnitTest value below, instead
diff --git a/Engine/Plugins/NetcodeUnitTest/NetcodeUnitTest/VAXSnippets/Unit Test Source.txt b/Engine/Plugins/NetcodeUnitTest/NetcodeUnitTest/VAXSnippets/Unit Test Source.txt
--- a/Engine/Plugins/NetcodeUnitTest/NetcodeUnitTest/VAXSnippets/Unit Test Source.txt	
+++ b/Engine/Plugins/NetcodeUnitTest/NetcodeUnitTest/VAXSnippets/Unit Test Source.txt	
@@ -39,6 +39,8 @@ void U$FILE_BASE$::InitializeEnvironmentSettings()
 {
 	BaseServerURL = UnitEnv->GetDefaultMap(UnitTestFlags);
 	BaseServerParameters = UnitEnv->GetDefaultServerParameters();
+
+	UnitEnv->InitializeUnitTasks();
 }
 
 void U$FILE_BASE$::ExecuteClientUnitTest()
diff --git a/Engine/Plugins/Online/OnlineFramework/Source/Lobby/Private/LobbyBeaconClient.cpp b/Engine/Plugins/Online/OnlineFramework/Source/Lobby/Private/LobbyBeaconClient.cpp
--- a/Engine/Plugins/Online/OnlineFramework/Source/Lobby/Private/LobbyBeaconClient.cpp
+++ b/Engine/Plugins/Online/OnlineFramework/Source/Lobby/Private/LobbyBeaconClient.cpp
@@ -286,7 +286,8 @@ bool ALobbyBeaconClient::ServerKickPlayer_Validate(const FUniqueNetIdRepl& Playe
 
 void ALobbyBeaconClient::ServerKickPlayer_Implementation(const FUniqueNetIdRepl& PlayerToKick, const FText& Reason)
 {
-	UE_LOG(LogBeacon, Log, TEXT("ServerKickPlayer %s -> %s"), *PlayerState->UniqueId.ToString(), *PlayerToKick.ToString());
+	UE_LOG(LogBeacon, Log, TEXT("ServerKickPlayer %s -> %s"), (PlayerState != nullptr ? *PlayerState->UniqueId.ToString() : TEXT("")),
+			*PlayerToKick.ToString());
 
 	ALobbyBeaconHost* BeaconHost = Cast<ALobbyBeaconHost>(GetBeaconOwner());
 	if (BeaconHost)
diff --git a/Engine/Plugins/Online/OnlineSubsystemSteam/Source/Private/OnlineAsyncTaskManagerSteam.cpp b/Engine/Plugins/Online/OnlineSubsystemSteam/Source/Private/OnlineAsyncTaskManagerSteam.cpp
--- a/Engine/Plugins/Online/OnlineSubsystemSteam/Source/Private/OnlineAsyncTaskManagerSteam.cpp
+++ b/Engine/Plugins/Online/OnlineSubsystemSteam/Source/Private/OnlineAsyncTaskManagerSteam.cpp
@@ -5,6 +5,7 @@
 #include "OnlineSubsystemSteam.h"
 #include "OnlineSubsystemSteamTypes.h"
 #include "OnlineSessionInterfaceSteam.h"
+#include "OnlinePresenceInterfaceSteam.h"
 #include "OnlineSessionAsyncLobbySteam.h"
 #include "OnlineSessionAsyncServerSteam.h"
 #include "OnlineLeaderboardInterfaceSteam.h"
@@ -1053,3 +1054,82 @@ void FOnlineAsyncTaskManagerSteam::OnSteamShutdown(SteamShutdown_t* CallbackData
 	UE_LOG_ONLINE(Verbose, TEXT("%s"), *NewEvent->ToString());
 	AddToOutQueue(NewEvent);
 }
+
+/**
+ * Notification event from Steam that rich presence has updated
+ */
+class FOnlineAsyncEventSteamRichPresenceUpdate : public FOnlineAsyncEvent<FOnlineSubsystemSteam>
+{
+	FOnlineAsyncEventSteamRichPresenceUpdate() :
+		FOnlineAsyncEvent(NULL)
+	{
+	}
+
+	FUniqueNetIdSteam TargetSteamId;
+
+public:
+
+	FOnlineAsyncEventSteamRichPresenceUpdate(FOnlineSubsystemSteam* InSubsystem, CSteamID InSteamId) :
+		FOnlineAsyncEvent(InSubsystem),
+		TargetSteamId(InSteamId)
+	{
+	}
+
+	FOnlineAsyncEventSteamRichPresenceUpdate(FOnlineSubsystemSteam* InSubsystem, uint64 InSteamId) :
+		FOnlineAsyncEvent(InSubsystem),
+		TargetSteamId(InSteamId)
+	{
+	}
+
+	/**
+	 * Get a human readable description of task
+	 */
+	virtual FString ToString() const override
+	{
+		return FString::Printf(TEXT("FOnlineAsyncEventSteamRichPresenceUpdate got new information about user %s"), *TargetSteamId.ToString());
+	}
+
+	/**
+	 * Give the async task a chance to marshal its data back to the game thread
+	 * Can only be called on the game thread by the async task manager
+	 */
+	virtual void Finalize() override
+	{
+		FOnlinePresenceSteamPtr PresenceInterface = StaticCastSharedPtr<FOnlinePresenceSteam>(Subsystem->GetPresenceInterface());
+		if (PresenceInterface.IsValid())
+		{
+			PresenceInterface->UpdatePresenceForUser(TargetSteamId);
+		}
+	}
+};
+
+/**
+ * Delegate registered with Steam to trigger when Steam gets updates about user rich presence
+ *
+ * @param CallbackData - Steam struct containing user that got their data updated
+ */
+void FOnlineAsyncTaskManagerSteam::OnRichPresenceUpdate(FriendRichPresenceUpdate_t* CallbackData)
+{
+	FOnlineAsyncEventSteamRichPresenceUpdate* NewEvent = new FOnlineAsyncEventSteamRichPresenceUpdate(SteamSubsystem, CallbackData->m_steamIDFriend);
+	UE_LOG_ONLINE(Verbose, TEXT("%s"), *NewEvent->ToString());
+	AddToOutQueue(NewEvent);
+}
+
+/**
+ * Delegate registered with Steam to trigger when Steam gets updates about user rich presence
+ *
+ * @param CallbackData - Steam struct containing user that got their data updated
+ */
+void FOnlineAsyncTaskManagerSteam::OnFriendStatusUpdate(PersonaStateChange_t* CallbackData)
+{
+	int ChangedData = CallbackData->m_nChangeFlags;
+	// Licensees can feel free to expand on this by adding their own watch events as well.
+	int RichPresenceWatchedEvents = (k_EPersonaChangeGameServer | k_EPersonaChangeGamePlayed | k_EPersonaChangeStatus | k_EPersonaChangeGoneOffline | k_EPersonaChangeComeOnline);
+	
+	if (ChangedData & RichPresenceWatchedEvents)
+	{
+		FOnlineAsyncEventSteamRichPresenceUpdate* NewEvent = new FOnlineAsyncEventSteamRichPresenceUpdate(SteamSubsystem, CallbackData->m_ulSteamID);
+		UE_LOG_ONLINE(Verbose, TEXT("%s"), *NewEvent->ToString());
+		AddToOutQueue(NewEvent);
+	}
+}
\ No newline at end of file
diff --git a/Engine/Plugins/Online/OnlineSubsystemSteam/Source/Private/OnlineAsyncTaskManagerSteam.h b/Engine/Plugins/Online/OnlineSubsystemSteam/Source/Private/OnlineAsyncTaskManagerSteam.h
--- a/Engine/Plugins/Online/OnlineSubsystemSteam/Source/Private/OnlineAsyncTaskManagerSteam.h
+++ b/Engine/Plugins/Online/OnlineSubsystemSteam/Source/Private/OnlineAsyncTaskManagerSteam.h
@@ -90,6 +90,10 @@ private:
 	STEAM_CALLBACK(FOnlineAsyncTaskManagerSteam, OnExternalUITriggered, GameOverlayActivated_t, OnExternalUITriggeredCallback);
 	/** Delegate registered with Steam to trigger when Steam is shutting down */
 	STEAM_CALLBACK(FOnlineAsyncTaskManagerSteam, OnSteamShutdown, SteamShutdown_t, OnSteamShutdownCallback);
+	/** Delegate registered with Steam to trigger when Steam gets updates about user rich presence */
+	STEAM_CALLBACK(FOnlineAsyncTaskManagerSteam, OnRichPresenceUpdate, FriendRichPresenceUpdate_t, OnRichPresenceUpdateCallback);
+	/** Delegate registered with Steam to trigger when Steam gets updates about friends status */
+	STEAM_CALLBACK(FOnlineAsyncTaskManagerSteam, OnFriendStatusUpdate, PersonaStateChange_t, OnFriendStatusUpdateCallback);
 
 	//GameServerChangeRequested_t
 
@@ -121,6 +125,8 @@ public:
 		OnUserStatsUnloadedCallback(this, &FOnlineAsyncTaskManagerSteam::OnUserStatsUnloaded),
 		OnExternalUITriggeredCallback(this, &FOnlineAsyncTaskManagerSteam::OnExternalUITriggered),
 		OnSteamShutdownCallback(this, &FOnlineAsyncTaskManagerSteam::OnSteamShutdown),
+		OnRichPresenceUpdateCallback(this, &FOnlineAsyncTaskManagerSteam::OnRichPresenceUpdate),
+		OnFriendStatusUpdateCallback(this, &FOnlineAsyncTaskManagerSteam::OnFriendStatusUpdate),
 		SteamSubsystem(InOnlineSubsystem)
 	{
 	}
diff --git a/Engine/Plugins/Online/OnlineSubsystemSteam/Source/Private/OnlineSessionInterfaceSteam.cpp b/Engine/Plugins/Online/OnlineSubsystemSteam/Source/Private/OnlineSessionInterfaceSteam.cpp
--- a/Engine/Plugins/Online/OnlineSubsystemSteam/Source/Private/OnlineSessionInterfaceSteam.cpp
+++ b/Engine/Plugins/Online/OnlineSubsystemSteam/Source/Private/OnlineSessionInterfaceSteam.cpp
@@ -626,6 +626,12 @@ uint32 FOnlineSessionSteam::DestroyInternetSession(FNamedOnlineSession* Session,
 		check(SessionInfo->SessionType == ESteamSession::AdvertisedSessionHost || SessionInfo->SessionType == ESteamSession::AdvertisedSessionClient);
 	}
 
+	// Clear any session advertisements this account had with this session.
+	if (SteamUser() != NULL && SteamUser()->BLoggedOn())
+	{
+		SteamUser()->AdvertiseGame(k_steamIDNil, 0, 0);
+	}
+
 	if (bSteamworksGameServerConnected && GameServerSteamId->IsValid())
 	{
 		// Logoff the master server
@@ -926,10 +932,22 @@ uint32 FOnlineSessionSteam::JoinInternetSession(int32 PlayerNum, FNamedOnlineSes
 			SteamSessionInfo->HostAddr = SearchSessionInfo->HostAddr;
 			SteamSessionInfo->SteamP2PAddr = SearchSessionInfo->SteamP2PAddr;
 
-			FString ConnectionString = GetSteamConnectionString(Session->SessionName);
-			if (!SteamFriends()->SetRichPresence("connect", TCHAR_TO_UTF8(*ConnectionString)))
+			if (SearchSession->SessionSettings.bAllowJoinViaPresence)
 			{
-				UE_LOG_ONLINE(Verbose, TEXT("Failed to set rich presence for session %s"), *Session->SessionName.ToString());
+				FString ConnectionString = GetSteamConnectionString(Session->SessionName);
+				if (!SteamFriends()->SetRichPresence("connect", TCHAR_TO_UTF8(*ConnectionString)))
+				{
+					UE_LOG_ONLINE(Verbose, TEXT("Failed to set rich presence for session %s"), *Session->SessionName.ToString());
+				}
+
+				// Advertise any servers we join
+				if (SteamUser() != NULL && SteamUser()->BLoggedOn())
+				{
+					uint32 IpAddr;
+					uint32 Port = SteamSessionInfo->HostAddr->GetPort();
+					SteamSessionInfo->HostAddr->GetIp(IpAddr);
+					SteamUser()->AdvertiseGame(SteamSessionInfo->SessionId, IpAddr, Port);
+				}
 			}
 			Result = ERROR_SUCCESS;
 		}
diff --git a/Engine/Plugins/Online/OnlineSubsystemSteam/Source/Private/OnlineSubsystemSteam.cpp b/Engine/Plugins/Online/OnlineSubsystemSteam/Source/Private/OnlineSubsystemSteam.cpp
--- a/Engine/Plugins/Online/OnlineSubsystemSteam/Source/Private/OnlineSubsystemSteam.cpp
+++ b/Engine/Plugins/Online/OnlineSubsystemSteam/Source/Private/OnlineSubsystemSteam.cpp
@@ -20,6 +20,7 @@
 #include "OnlineAsyncTaskManagerSteam.h"
 #include "OnlineSessionInterfaceSteam.h"
 #include "OnlineIdentityInterfaceSteam.h"
+#include "OnlinePresenceInterfaceSteam.h"
 #include "OnlineFriendsInterfaceSteam.h"
 #include "OnlineUserCloudInterfaceSteam.h"
 #include "OnlineSharedCloudInterfaceSteam.h"
@@ -28,6 +29,27 @@
 #include "OnlineExternalUIInterfaceSteam.h"
 #include "OnlineAchievementsInterfaceSteam.h"
 
+
+#if !UE_BUILD_SHIPPING
+namespace OSSConsoleVariables
+{
+	/** This CVar is for use by NetcodeUnitTest, and is completely unsupported for use outside of this */
+	TAutoConsoleVariable<int32> CVarSteamInitServerOnClient(
+			TEXT("OSS.SteamInitServerOnClient"),
+			1,
+			TEXT("Whether or not to initialize the Steam server interface on clients (default true)"),
+			ECVF_Default);
+
+	/** CVar used by NetcodeUnitTest, to force-enable Steam within the unit test commandlet */
+	TAutoConsoleVariable<int32> CVarSteamUnitTest(
+			TEXT("OSS.SteamUnitTest"),
+			0,
+			TEXT("Whether or not Steam is being force-enabled by NetcodeUnitTest"),
+			ECVF_Default);
+}
+#endif
+
+
 extern "C" 
 { 
 	static void __cdecl SteamworksWarningMessageHook(int Severity, const char *Message); 
@@ -266,7 +288,7 @@ IOnlineMessagePtr FOnlineSubsystemSteam::GetMessageInterface() const
 
 IOnlinePresencePtr FOnlineSubsystemSteam::GetPresenceInterface() const
 {
-	return nullptr;
+	return PresenceInterface;
 }
 
 IOnlineChatPtr FOnlineSubsystemSteam::GetChatInterface() const
@@ -322,13 +344,20 @@ bool FOnlineSubsystemSteam::Init()
 	ConfigureSteamInitDevOptions(bRelaunchInSteam, RelaunchAppId);
 
 	const bool bIsServer = IsRunningDedicatedServer();
+	bool bAttemptServerInit = true;
+
+#if !UE_BUILD_SHIPPING
+	// Add a bypass for NetcodeUnitTest, to allow running a Steam server + client from same machine, by disabling server init on client.
+	// This is an unapproved/unsupported method for using OnlineSubsystemSteam.
+	bAttemptServerInit = bIsServer || !!OSSConsoleVariables::CVarSteamInitServerOnClient.GetValueOnGameThread();
+#endif
 	
 	// Don't initialize the Steam Client API if we are launching as a server
 	bool bClientInitSuccess = !bIsServer ? InitSteamworksClient(bRelaunchInSteam, RelaunchAppId) : true;
 
 	// Initialize the Steam Server API if this is a dedicated server or
 	//  the Client API was successfully initialized
-	bool bServerInitSuccess = bClientInitSuccess ? InitSteamworksServer() : false;
+	bool bServerInitSuccess = bClientInitSuccess ? (!bAttemptServerInit || InitSteamworksServer()) : false;
 
 	if (bClientInitSuccess && bServerInitSuccess)
 	{
@@ -346,6 +375,8 @@ bool FOnlineSubsystemSteam::Init()
 
 		IdentityInterface = MakeShareable(new FOnlineIdentitySteam(this));
 
+		PresenceInterface = MakeShareable(new FOnlinePresenceSteam(this));
+
 		if (!bIsServer)
 		{
 			FriendInterface = MakeShareable(new FOnlineFriendsSteam(this));
@@ -418,6 +449,7 @@ bool FOnlineSubsystemSteam::Shutdown()
 	DESTRUCT_INTERFACE(FriendInterface);
 	DESTRUCT_INTERFACE(IdentityInterface);
 	DESTRUCT_INTERFACE(SessionInterface);
+	DESTRUCT_INTERFACE(PresenceInterface);
 
 #undef DESTRUCT_INTERFACE
 
@@ -453,6 +485,14 @@ bool FOnlineSubsystemSteam::IsEnabled() const
 			bEnableSteam = IsRunningDedicatedServer() || IsRunningGame();
 		}
 #endif
+
+#if !UE_BUILD_SHIPPING
+		// Force-enable Steam for NetcodeUnitTest
+		if (!!OSSConsoleVariables::CVarSteamUnitTest.GetValueOnGameThread())
+		{
+			bEnableSteam = true;
+		}
+#endif
 	}
 
 	return bEnableSteam;
diff --git a/Engine/Plugins/Online/OnlineSubsystemSteam/Source/Private/OnlineUserCloudInterfaceSteam.cpp b/Engine/Plugins/Online/OnlineSubsystemSteam/Source/Private/OnlineUserCloudInterfaceSteam.cpp
--- a/Engine/Plugins/Online/OnlineSubsystemSteam/Source/Private/OnlineUserCloudInterfaceSteam.cpp
+++ b/Engine/Plugins/Online/OnlineSubsystemSteam/Source/Private/OnlineUserCloudInterfaceSteam.cpp
@@ -246,7 +246,7 @@ FOnlineUserCloudSteam::~FOnlineUserCloudSteam()
 
 bool FOnlineUserCloudSteam::GetFileContents(const FUniqueNetId& UserId, const FString& FileName, TArray<uint8>& FileContents)
 {
-	FScopeLock(&SteamSubsystem->UserCloudDataLock);
+	FScopeLock ScopeLock(&SteamSubsystem->UserCloudDataLock);
 	// Search for the specified file and return the raw data
 	FSteamUserCloudData* UserCloudData = SteamSubsystem->GetUserCloudEntry(UserId);
 	if (UserCloudData)
@@ -264,7 +264,7 @@ bool FOnlineUserCloudSteam::GetFileContents(const FUniqueNetId& UserId, const FS
 
 bool FOnlineUserCloudSteam::ClearFiles(const FUniqueNetId& UserId)
 {
-	FScopeLock(&SteamSubsystem->UserCloudDataLock);
+	FScopeLock ScopeLock(&SteamSubsystem->UserCloudDataLock);
 	// Search for the specified file
 	FSteamUserCloudData* UserCloudData = SteamSubsystem->GetUserCloudEntry(UserId);
 	if (UserCloudData)
@@ -277,7 +277,7 @@ bool FOnlineUserCloudSteam::ClearFiles(const FUniqueNetId& UserId)
 
 bool FOnlineUserCloudSteam::ClearFile(const FUniqueNetId& UserId, const FString& FileName)
 {
-	FScopeLock(&SteamSubsystem->UserCloudDataLock);
+	FScopeLock ScopeLock(&SteamSubsystem->UserCloudDataLock);
 	// Search for the specified file
 	FSteamUserCloudData* UserCloudData = SteamSubsystem->GetUserCloudEntry(UserId);
 	if (UserCloudData)
@@ -295,14 +295,14 @@ void FOnlineUserCloudSteam::EnumerateUserFiles(const FUniqueNetId& UserId)
 
 void FOnlineUserCloudSteam::GetUserFileList(const FUniqueNetId& UserId, TArray<FCloudFileHeader>& UserFiles)
 {
-	FScopeLock(&SteamSubsystem->UserCloudDataLock);
+	FScopeLock ScopeLock(&SteamSubsystem->UserCloudDataLock);
 	FSteamUserCloudData* UserMetadata = SteamSubsystem->GetUserCloudEntry(UserId);  
 	UserFiles = UserMetadata->CloudMetadata;
 }
 
 bool FOnlineUserCloudSteam::ReadUserFile(const FUniqueNetId& UserId, const FString& FileName)
 {
-	FScopeLock(&SteamSubsystem->UserCloudDataLock);
+	FScopeLock ScopeLock(&SteamSubsystem->UserCloudDataLock);
 	// Create or get the current entry for this file
 	FSteamUserCloudData* UserCloud = SteamSubsystem->GetUserCloudEntry(UserId);
 	if (UserCloud && FileName.Len() > 0)
@@ -318,7 +318,7 @@ bool FOnlineUserCloudSteam::ReadUserFile(const FUniqueNetId& UserId, const FStri
 
 bool FOnlineUserCloudSteam::WriteUserFile(const FUniqueNetId& UserId, const FString& FileName, TArray<uint8>& FileContents)
 {
-	FScopeLock(&SteamSubsystem->UserCloudDataLock);
+	FScopeLock ScopeLock(&SteamSubsystem->UserCloudDataLock);
 	// Create or get the current entry for this file
 	FSteamUserCloudData* UserCloud = SteamSubsystem->GetUserCloudEntry(UserId);
 	if (UserCloud && FileName.Len() > 0)
@@ -369,7 +369,7 @@ void FOnlineUserCloudSteam::DumpCloudFileState(const FUniqueNetId& UserId, const
 	{
 		UE_LOG_ONLINE(Log, TEXT("Cloud File State file %s:"), *FileName);
 		{
-			FScopeLock(&SteamSubsystem->UserCloudDataLock);
+			FScopeLock ScopeLock(&SteamSubsystem->UserCloudDataLock);
 			FSteamUserCloudData* UserCloud = SteamSubsystem->GetUserCloudEntry(UserId);
 			FCloudFileHeader* FileMetadata = UserCloud->GetFileMetadata(FileName);
 			if (FileMetadata)
diff --git a/Engine/Plugins/Online/OnlineSubsystemSteam/Source/Private/VoiceEngineSteam.cpp b/Engine/Plugins/Online/OnlineSubsystemSteam/Source/Private/VoiceEngineSteam.cpp
--- a/Engine/Plugins/Online/OnlineSubsystemSteam/Source/Private/VoiceEngineSteam.cpp
+++ b/Engine/Plugins/Online/OnlineSubsystemSteam/Source/Private/VoiceEngineSteam.cpp
@@ -24,6 +24,8 @@ FVoiceEngineSteam::FVoiceEngineSteam(FOnlineSubsystemSteam* InSteamSubsystem) :
 {
 	SteamUserPtr = SteamUser();
 	SteamFriendsPtr = SteamFriends();
+
+	FCoreUObjectDelegates::PostLoadMapWithWorld.AddRaw(this, &FVoiceEngineSteam::OnPostLoadMap);
 }
 
 FVoiceEngineSteam::~FVoiceEngineSteam()
@@ -33,7 +35,14 @@ FVoiceEngineSteam::~FVoiceEngineSteam()
 		SteamUserPtr->StopVoiceRecording();
 		SteamFriendsPtr->SetInGameVoiceSpeaking(SteamUserPtr->GetSteamID(), false);
 	}
-	delete SerializeHelper;
+
+	if (SerializeHelper != nullptr)
+	{
+		delete SerializeHelper;
+		SerializeHelper = nullptr;
+	}
+
+	FCoreUObjectDelegates::PostLoadMapWithWorld.RemoveAll(this);
 }
 
 void FVoiceEngineSteam::VoiceCaptureUpdate() const
@@ -379,6 +388,18 @@ void FVoiceEngineSteam::OnAudioFinished(UAudioComponent* AC)
 	UE_LOG(LogVoiceEngine, Verbose, TEXT("Audio Finished"));
 }
 
+void FVoiceEngineSteam::OnPostLoadMap(UWorld*)
+{
+	for (FRemoteTalkerData::TIterator It(RemoteTalkerBuffers); It; ++It)
+	{
+		FRemoteTalkerDataSteam& RemoteData = It.Value();
+		if (RemoteData.AudioComponent)
+		{
+			RemoteData.AudioComponent->Play();
+		}
+	}
+}
+
 FString FVoiceEngineSteam::GetVoiceDebugState() const 
 {
 	FString Output;
diff --git a/Engine/Plugins/Online/OnlineSubsystemSteam/Source/Public/OnlineSubsystemSteam.h b/Engine/Plugins/Online/OnlineSubsystemSteam/Source/Public/OnlineSubsystemSteam.h
--- a/Engine/Plugins/Online/OnlineSubsystemSteam/Source/Public/OnlineSubsystemSteam.h
+++ b/Engine/Plugins/Online/OnlineSubsystemSteam/Source/Public/OnlineSubsystemSteam.h
@@ -15,6 +15,7 @@ class FOnlineSessionSteam;
 class FOnlineSharedCloudSteam;
 class FOnlineUserCloudSteam;
 class FOnlineVoiceSteam;
+class FOnlinePresenceSteam;
 
 /** Forward declarations of all interface classes */
 typedef TSharedPtr<class FOnlineSessionSteam, ESPMode::ThreadSafe> FOnlineSessionSteamPtr;
@@ -26,6 +27,7 @@ typedef TSharedPtr<class FOnlineLeaderboardsSteam, ESPMode::ThreadSafe> FOnlineL
 typedef TSharedPtr<class FOnlineVoiceSteam, ESPMode::ThreadSafe> FOnlineVoiceSteamPtr;
 typedef TSharedPtr<class FOnlineExternalUISteam, ESPMode::ThreadSafe> FOnlineExternalUISteamPtr;
 typedef TSharedPtr<class FOnlineAchievementsSteam, ESPMode::ThreadSafe> FOnlineAchievementsSteamPtr;
+typedef TSharedPtr<class FOnlinePresenceSteam, ESPMode::ThreadSafe> FOnlinePresenceSteamPtr;
 
 /**
  *	OnlineSubsystemSteam - Implementation of the online subsystem for STEAM services
@@ -83,6 +85,9 @@ protected:
 	/** Interface for achievements */
 	FOnlineAchievementsSteamPtr AchievementsInterface;
 
+	/** Interface for presence */
+	FOnlinePresenceSteamPtr PresenceInterface;
+
 	/** Online async task runnable */
 	class FOnlineAsyncTaskManagerSteam* OnlineAsyncTaskThreadRunnable;
 
@@ -108,6 +113,7 @@ PACKAGE_SCOPE:
 		LeaderboardsInterface(nullptr),
 		VoiceInterface(nullptr),
 		ExternalUIInterface(nullptr),
+		PresenceInterface(nullptr),
 		OnlineAsyncTaskThreadRunnable(nullptr),
 		OnlineAsyncTaskThread(nullptr)
 	{}
@@ -127,6 +133,7 @@ PACKAGE_SCOPE:
 		LeaderboardsInterface(nullptr),
 		VoiceInterface(nullptr),
 		ExternalUIInterface(nullptr),
+		PresenceInterface(nullptr),
 		OnlineAsyncTaskThreadRunnable(nullptr),
 		OnlineAsyncTaskThread(nullptr)
 	{}
diff --git a/Engine/Plugins/Online/OnlineSubsystemUtils/Source/OnlineSubsystemUtils/Private/OnlineBeaconClient.cpp b/Engine/Plugins/Online/OnlineSubsystemUtils/Source/OnlineSubsystemUtils/Private/OnlineBeaconClient.cpp
--- a/Engine/Plugins/Online/OnlineSubsystemUtils/Source/OnlineSubsystemUtils/Private/OnlineBeaconClient.cpp
+++ b/Engine/Plugins/Online/OnlineSubsystemUtils/Source/OnlineSubsystemUtils/Private/OnlineBeaconClient.cpp
@@ -323,49 +323,63 @@ void AOnlineBeaconClient::NotifyControlMessage(UNetConnection* Connection, uint8
 		case NMT_BeaconAssignGUID:
 			{
 				FNetworkGUID NetGUID;
-				FNetControlMessage<NMT_BeaconAssignGUID>::Receive(Bunch, NetGUID);
-				if (NetGUID.IsValid())
-				{
-					Connection->Driver->GuidCache->RegisterNetGUID_Client( NetGUID, this );
-
-					FString BeaconType = GetBeaconType();
-					FNetControlMessage<NMT_BeaconNetGUIDAck>::Send(Connection, BeaconType);
-					// Server will send ClientOnConnected() when it gets this control message
 
-					// Fail safe for connection to server but no client connection RPC
-					FTimerDelegate TimerDelegate = FTimerDelegate::CreateUObject(this, &AOnlineBeaconClient::OnFailure);
-					GetWorldTimerManager().SetTimer(TimerHandle_OnFailure, TimerDelegate, BEACON_RPC_TIMEOUT, false);
-				}
-				else
+				if (FNetControlMessage<NMT_BeaconAssignGUID>::Receive(Bunch, NetGUID))
 				{
-					// Force close the session
-					UE_LOG(LogBeacon, Log, TEXT("Beacon close from invalid NetGUID"));
-					OnFailure();
+					if (NetGUID.IsValid())
+					{
+						Connection->Driver->GuidCache->RegisterNetGUID_Client(NetGUID, this);
+
+						FString BeaconType = GetBeaconType();
+						FNetControlMessage<NMT_BeaconNetGUIDAck>::Send(Connection, BeaconType);
+						// Server will send ClientOnConnected() when it gets this control message
+
+						// Fail safe for connection to server but no client connection RPC
+						FTimerDelegate TimerDelegate = FTimerDelegate::CreateUObject(this, &AOnlineBeaconClient::OnFailure);
+						GetWorldTimerManager().SetTimer(TimerHandle_OnFailure, TimerDelegate, BEACON_RPC_TIMEOUT, false);
+					}
+					else
+					{
+						// Force close the session
+						UE_LOG(LogBeacon, Log, TEXT("Beacon close from invalid NetGUID"));
+						OnFailure();
+					}
 				}
+
 				break;
 			}
 		case NMT_Upgrade:
 			{
 				// Report mismatch.
 				uint32 RemoteNetworkVersion;
-				FNetControlMessage<NMT_Upgrade>::Receive(Bunch, RemoteNetworkVersion);
-				// Upgrade
-				const FString ConnectionError = NSLOCTEXT("Engine", "ClientOutdated", "The match you are trying to join is running an incompatible version of the game.  Please try upgrading your game version.").ToString();
-				GEngine->BroadcastNetworkFailure(GetWorld(), NetDriver, ENetworkFailure::OutdatedClient, ConnectionError);
+
+				if (FNetControlMessage<NMT_Upgrade>::Receive(Bunch, RemoteNetworkVersion))
+				{
+					// Upgrade
+					const FString ConnectionError = NSLOCTEXT("Engine", "ClientOutdated",
+						"The match you are trying to join is running an incompatible version of the game.  Please try upgrading your game version.").ToString();
+
+					GEngine->BroadcastNetworkFailure(GetWorld(), NetDriver, ENetworkFailure::OutdatedClient, ConnectionError);
+				}
+
 				break;
 			}
 		case NMT_Failure:
 			{
 				FString ErrorMsg;
-				FNetControlMessage<NMT_Failure>::Receive(Bunch, ErrorMsg);
-				if (ErrorMsg.IsEmpty())
+
+				if (FNetControlMessage<NMT_Failure>::Receive(Bunch, ErrorMsg))
 				{
-					ErrorMsg = NSLOCTEXT("NetworkErrors", "GenericBeaconConnectionFailed", "Beacon Connection Failed.").ToString();
+					if (ErrorMsg.IsEmpty())
+					{
+						ErrorMsg = NSLOCTEXT("NetworkErrors", "GenericBeaconConnectionFailed", "Beacon Connection Failed.").ToString();
+					}
+
+					// Force close the session
+					UE_LOG(LogBeacon, Log, TEXT("Beacon close from NMT_Failure %s"), *ErrorMsg);
+					OnFailure();
 				}
 
-				// Force close the session
-				UE_LOG(LogBeacon, Log, TEXT("Beacon close from NMT_Failure %s"), *ErrorMsg);
-				OnFailure();
 				break;
 			}
 		case NMT_BeaconJoin:
diff --git a/Engine/Plugins/Online/OnlineSubsystemUtils/Source/OnlineSubsystemUtils/Private/OnlineBeaconHost.cpp b/Engine/Plugins/Online/OnlineSubsystemUtils/Source/OnlineSubsystemUtils/Private/OnlineBeaconHost.cpp
--- a/Engine/Plugins/Online/OnlineSubsystemUtils/Source/OnlineSubsystemUtils/Private/OnlineBeaconHost.cpp
+++ b/Engine/Plugins/Online/OnlineSubsystemUtils/Source/OnlineSubsystemUtils/Private/OnlineBeaconHost.cpp
@@ -94,44 +94,50 @@ void AOnlineBeaconHost::NotifyControlMessage(UNetConnection* Connection, uint8 M
 				uint32 LocalNetworkVersion = FNetworkVersion::GetLocalNetworkVersion();
 				FString EncryptionToken;
 
-				FNetControlMessage<NMT_Hello>::Receive(Bunch, IsLittleEndian, RemoteNetworkVersion, EncryptionToken);
-
-				if (!FNetworkVersion::IsNetworkCompatible(LocalNetworkVersion, RemoteNetworkVersion))
-				{
-					UE_LOG(LogBeacon, Error, TEXT("Client not network compatible %s (Local=%d, Remote=%d)"), *Connection->GetName(), LocalNetworkVersion, RemoteNetworkVersion);
-					FNetControlMessage<NMT_Upgrade>::Send(Connection, LocalNetworkVersion);
-					bCloseConnection = true;
-				}
-				else
+				if (FNetControlMessage<NMT_Hello>::Receive(Bunch, IsLittleEndian, RemoteNetworkVersion, EncryptionToken))
 				{
-					if (EncryptionToken.IsEmpty())
+					if (!FNetworkVersion::IsNetworkCompatible(LocalNetworkVersion, RemoteNetworkVersion))
 					{
-						SendWelcomeControlMessage(Connection);
+						UE_LOG(LogBeacon, Error, TEXT("Client not network compatible %s (Local=%d, Remote=%d)"), *Connection->GetName(), LocalNetworkVersion, RemoteNetworkVersion);
+						FNetControlMessage<NMT_Upgrade>::Send(Connection, LocalNetworkVersion);
+						bCloseConnection = true;
 					}
 					else
 					{
-						if (FNetDelegates::OnReceivedNetworkEncryptionToken.IsBound())
+						if (EncryptionToken.IsEmpty())
 						{
-							TWeakObjectPtr<UNetConnection> WeakConnection = Connection;
-							FNetDelegates::OnReceivedNetworkEncryptionToken.Execute(EncryptionToken, FOnEncryptionKeyResponse::CreateUObject(this, &AOnlineBeaconHost::SendWelcomeControlMessage, WeakConnection));
+							SendWelcomeControlMessage(Connection);
 						}
 						else
 						{
-							FString FailureMsg(TEXT("Encryption failure"));
-							UE_LOG(LogBeacon, Warning, TEXT("%s: No delegate available to handle encryption token, disconnecting."), *Connection->GetName());
-							FNetControlMessage<NMT_Failure>::Send(Connection, FailureMsg);
-							bCloseConnection = true;
+							if (FNetDelegates::OnReceivedNetworkEncryptionToken.IsBound())
+							{
+								TWeakObjectPtr<UNetConnection> WeakConnection = Connection;
+								FNetDelegates::OnReceivedNetworkEncryptionToken.Execute(EncryptionToken, FOnEncryptionKeyResponse::CreateUObject(this, &AOnlineBeaconHost::SendWelcomeControlMessage, WeakConnection));
+							}
+							else
+							{
+								FString FailureMsg(TEXT("Encryption failure"));
+								UE_LOG(LogBeacon, Warning, TEXT("%s: No delegate available to handle encryption token, disconnecting."), *Connection->GetName());
+								FNetControlMessage<NMT_Failure>::Send(Connection, FailureMsg);
+								bCloseConnection = true;
+							}
 						}
 					}
 				}
+
 				break;
 			}
 		case NMT_Netspeed:
 			{
 				int32 Rate;
-				FNetControlMessage<NMT_Netspeed>::Receive(Bunch, Rate);
-				Connection->CurrentNetSpeed = FMath::Clamp(Rate, 1800, NetDriver->MaxClientRate);
-				UE_LOG(LogBeacon, Log, TEXT("Client netspeed is %i"), Connection->CurrentNetSpeed);
+
+				if (FNetControlMessage<NMT_Netspeed>::Receive(Bunch, Rate))
+				{
+					Connection->CurrentNetSpeed = FMath::Clamp(Rate, 1800, NetDriver->MaxClientRate);
+					UE_LOG(LogBeacon, Log, TEXT("Client netspeed is %i"), Connection->CurrentNetSpeed);
+				}
+
 				break;
 			}
 		case NMT_BeaconJoin:
@@ -139,60 +145,63 @@ void AOnlineBeaconHost::NotifyControlMessage(UNetConnection* Connection, uint8 M
 				FString ErrorMsg;
 				FString BeaconType;
 				FUniqueNetIdRepl UniqueId;
-				FNetControlMessage<NMT_BeaconJoin>::Receive(Bunch, BeaconType, UniqueId);
-				UE_LOG(LogBeacon, Log, TEXT("Beacon Join %s %s"), *BeaconType, *UniqueId.ToDebugString());
 
-				if (Connection->GetClientWorldPackageName() == NAME_None)
+				if (FNetControlMessage<NMT_BeaconJoin>::Receive(Bunch, BeaconType, UniqueId))
 				{
-					AOnlineBeaconClient* ClientActor = GetClientActor(Connection);
-					if (ClientActor == nullptr)
-					{
-						UWorld* World = GetWorld();
-						Connection->SetClientWorldPackageName(World->GetOutermost()->GetFName());
-
-						AOnlineBeaconClient* NewClientActor = nullptr;
-						FOnBeaconSpawned* OnBeaconSpawnedDelegate = OnBeaconSpawnedMapping.Find(BeaconType);
-						if (OnBeaconSpawnedDelegate && OnBeaconSpawnedDelegate->IsBound())
-						{
-							NewClientActor = OnBeaconSpawnedDelegate->Execute(Connection);
-						}
+					UE_LOG(LogBeacon, Log, TEXT("Beacon Join %s %s"), *BeaconType, *UniqueId.ToDebugString());
 
-						if (NewClientActor && BeaconType == NewClientActor->GetBeaconType())
+					if (Connection->GetClientWorldPackageName() == NAME_None)
+					{
+						AOnlineBeaconClient* ClientActor = GetClientActor(Connection);
+						if (ClientActor == nullptr)
 						{
-							NewClientActor->SetConnectionState(EBeaconConnectionState::Pending);
-
-							FNetworkGUID NetGUID = Connection->Driver->GuidCache->AssignNewNetGUID_Server(NewClientActor);
-							NewClientActor->SetNetConnection(Connection);
-							Connection->PlayerId = UniqueId;
-							Connection->OwningActor = NewClientActor;
-							Connection->SetClientLoginState(EClientLoginState::ReceivedJoin);
-							NewClientActor->Role = ROLE_Authority;
-							NewClientActor->SetReplicates(false);
-							check(NetDriverName == NetDriver->NetDriverName);
-							NewClientActor->SetNetDriverName(NetDriverName);
-							ClientActors.Add(NewClientActor);
-							FNetControlMessage<NMT_BeaconAssignGUID>::Send(Connection, NetGUID);
+							UWorld* World = GetWorld();
+							Connection->SetClientWorldPackageName(World->GetOutermost()->GetFName());
+
+							AOnlineBeaconClient* NewClientActor = nullptr;
+							FOnBeaconSpawned* OnBeaconSpawnedDelegate = OnBeaconSpawnedMapping.Find(BeaconType);
+							if (OnBeaconSpawnedDelegate && OnBeaconSpawnedDelegate->IsBound())
+							{
+								NewClientActor = OnBeaconSpawnedDelegate->Execute(Connection);
+							}
+
+							if (NewClientActor && BeaconType == NewClientActor->GetBeaconType())
+							{
+								NewClientActor->SetConnectionState(EBeaconConnectionState::Pending);
+
+								FNetworkGUID NetGUID = Connection->Driver->GuidCache->AssignNewNetGUID_Server(NewClientActor);
+								NewClientActor->SetNetConnection(Connection);
+								Connection->PlayerId = UniqueId;
+								Connection->OwningActor = NewClientActor;
+								Connection->SetClientLoginState(EClientLoginState::ReceivedJoin);
+								NewClientActor->Role = ROLE_Authority;
+								NewClientActor->SetReplicates(false);
+								check(NetDriverName == NetDriver->NetDriverName);
+								NewClientActor->SetNetDriverName(NetDriverName);
+								ClientActors.Add(NewClientActor);
+								FNetControlMessage<NMT_BeaconAssignGUID>::Send(Connection, NetGUID);
+							}
+							else
+							{
+								ErrorMsg = NSLOCTEXT("NetworkErrors", "BeaconSpawnFailureError", "Join failure, Couldn't spawn beacon.").ToString();
+							}
 						}
 						else
 						{
-							ErrorMsg = NSLOCTEXT("NetworkErrors", "BeaconSpawnFailureError", "Join failure, Couldn't spawn beacon.").ToString();
+							ErrorMsg = NSLOCTEXT("NetworkErrors", "BeaconSpawnExistingActorError", "Join failure, existing beacon actor.").ToString();
 						}
 					}
 					else
 					{
-						ErrorMsg = NSLOCTEXT("NetworkErrors", "BeaconSpawnExistingActorError", "Join failure, existing beacon actor.").ToString();
+						ErrorMsg = NSLOCTEXT("NetworkErrors", "BeaconSpawnClientWorldPackageNameError", "Join failure, existing ClientWorldPackageName.").ToString();
 					}
-				}
-				else
-				{
-					ErrorMsg = NSLOCTEXT("NetworkErrors", "BeaconSpawnClientWorldPackageNameError", "Join failure, existing ClientWorldPackageName.").ToString();
-				}
 
-				if (!ErrorMsg.IsEmpty())
-				{
-					UE_LOG(LogBeacon, Log, TEXT("%s: %s"), *Connection->GetName(), *ErrorMsg);
-					FNetControlMessage<NMT_Failure>::Send(Connection, ErrorMsg);
-					bCloseConnection = true;
+					if (!ErrorMsg.IsEmpty())
+					{
+						UE_LOG(LogBeacon, Log, TEXT("%s: %s"), *Connection->GetName(), *ErrorMsg);
+						FNetControlMessage<NMT_Failure>::Send(Connection, ErrorMsg);
+						bCloseConnection = true;
+					}
 				}
 
 				break;
@@ -201,40 +210,42 @@ void AOnlineBeaconHost::NotifyControlMessage(UNetConnection* Connection, uint8 M
 			{
 				FString ErrorMsg;
 				FString BeaconType;
-				FNetControlMessage<NMT_BeaconNetGUIDAck>::Receive(Bunch, BeaconType);
 
-				AOnlineBeaconClient* ClientActor = GetClientActor(Connection);
-				if (ClientActor && BeaconType == ClientActor->GetBeaconType())
+				if (FNetControlMessage<NMT_BeaconNetGUIDAck>::Receive(Bunch, BeaconType))
 				{
-					FOnBeaconConnected* OnBeaconConnectedDelegate = OnBeaconConnectedMapping.Find(BeaconType);
-					if (OnBeaconConnectedDelegate)
+					AOnlineBeaconClient* ClientActor = GetClientActor(Connection);
+					if (ClientActor && BeaconType == ClientActor->GetBeaconType())
 					{
-						ClientActor->SetReplicates(true);
-						ClientActor->SetAutonomousProxy(true);
-						ClientActor->SetConnectionState(EBeaconConnectionState::Open);
-						// Send an RPC to the client to open the actor channel and guarantee RPCs will work
-						ClientActor->ClientOnConnected();
-						UE_LOG(LogBeacon, Log, TEXT("Handshake complete for %s!"), *ClientActor->GetName());
-
-						OnBeaconConnectedDelegate->ExecuteIfBound(ClientActor, Connection);
+						FOnBeaconConnected* OnBeaconConnectedDelegate = OnBeaconConnectedMapping.Find(BeaconType);
+						if (OnBeaconConnectedDelegate)
+						{
+							ClientActor->SetReplicates(true);
+							ClientActor->SetAutonomousProxy(true);
+							ClientActor->SetConnectionState(EBeaconConnectionState::Open);
+							// Send an RPC to the client to open the actor channel and guarantee RPCs will work
+							ClientActor->ClientOnConnected();
+							UE_LOG(LogBeacon, Log, TEXT("Handshake complete for %s!"), *ClientActor->GetName());
+
+							OnBeaconConnectedDelegate->ExecuteIfBound(ClientActor, Connection);
+						}
+						else
+						{
+							// Failed to connect.
+							ErrorMsg = NSLOCTEXT("NetworkErrors", "BeaconSpawnNetGUIDAckError1", "Join failure, no host object at NetGUIDAck.").ToString();
+						}
 					}
 					else
 					{
 						// Failed to connect.
-						ErrorMsg = NSLOCTEXT("NetworkErrors", "BeaconSpawnNetGUIDAckError1", "Join failure, no host object at NetGUIDAck.").ToString();
+						ErrorMsg = NSLOCTEXT("NetworkErrors", "BeaconSpawnNetGUIDAckError2", "Join failure, no actor at NetGUIDAck.").ToString();
 					}
-				}
-				else
-				{
-					// Failed to connect.
-					ErrorMsg = NSLOCTEXT("NetworkErrors", "BeaconSpawnNetGUIDAckError2", "Join failure, no actor at NetGUIDAck.").ToString();
-				}
 
-				if (!ErrorMsg.IsEmpty())
-				{
-					UE_LOG(LogBeacon, Log, TEXT("%s: %s"), *Connection->GetName(), *ErrorMsg);
-					FNetControlMessage<NMT_Failure>::Send(Connection, ErrorMsg);
-					bCloseConnection = true;
+					if (!ErrorMsg.IsEmpty())
+					{
+						UE_LOG(LogBeacon, Log, TEXT("%s: %s"), *Connection->GetName(), *ErrorMsg);
+						FNetControlMessage<NMT_Failure>::Send(Connection, ErrorMsg);
+						bCloseConnection = true;
+					}
 				}
 
 				break;
diff --git a/Engine/Plugins/Online/OnlineSubsystemUtils/Source/OnlineSubsystemUtils/Private/OnlineEngineInterfaceImpl.cpp b/Engine/Plugins/Online/OnlineSubsystemUtils/Source/OnlineSubsystemUtils/Private/OnlineEngineInterfaceImpl.cpp
--- a/Engine/Plugins/Online/OnlineSubsystemUtils/Source/OnlineSubsystemUtils/Private/OnlineEngineInterfaceImpl.cpp
+++ b/Engine/Plugins/Online/OnlineSubsystemUtils/Source/OnlineSubsystemUtils/Private/OnlineEngineInterfaceImpl.cpp
@@ -360,7 +360,12 @@ TSharedPtr<FVoicePacket> UOnlineEngineInterfaceImpl::SerializeRemotePacket(UWorl
 	FName VoiceSubsystemName = VoiceSubsystemNameOverride;
 	if (RemoteConnection && RemoteConnection->Driver && RemoteConnection->Driver->GetNetMode() == NM_DedicatedServer)
 	{
-		VoiceSubsystemName = GetDedicatedServerSubsystemNameForSubsystem(RemoteConnection->GetPlayerOnlinePlatformName());
+		FName DedicatedVoiceSubsystemName = GetDedicatedServerSubsystemNameForSubsystem(RemoteConnection->GetPlayerOnlinePlatformName());
+
+		if (DedicatedVoiceSubsystemName != NAME_None)
+		{
+			VoiceSubsystemName = DedicatedVoiceSubsystemName;
+		}
 	}
 
 	IOnlineVoicePtr VoiceInt = Online::GetVoiceInterface(World, VoiceSubsystemName);
diff --git a/Engine/Plugins/Online/OnlineSubsystemUtils/Source/OnlineSubsystemUtils/Private/OnlineSubsystemUtils.cpp b/Engine/Plugins/Online/OnlineSubsystemUtils/Source/OnlineSubsystemUtils/Private/OnlineSubsystemUtils.cpp
--- a/Engine/Plugins/Online/OnlineSubsystemUtils/Source/OnlineSubsystemUtils/Private/OnlineSubsystemUtils.cpp
+++ b/Engine/Plugins/Online/OnlineSubsystemUtils/Source/OnlineSubsystemUtils/Private/OnlineSubsystemUtils.cpp
@@ -32,6 +32,7 @@
 #include "Tests/TestMessageInterface.h"
 #include "Tests/TestVoice.h"
 #include "Tests/TestExternalUIInterface.h"
+#include "Tests/TestPresenceInterface.h"
 
 UAudioComponent* CreateVoiceAudioComponent(uint32 SampleRate, int32 NumChannels)
 {
@@ -460,6 +461,13 @@ static bool OnlineExec( UWorld* InWorld, const TCHAR* Cmd, FOutputDevice& Ar )
 						(new FTestLeaderboardInterface(SubName))->Test(InWorld);
 						bWasHandled = true;
 					}
+					else if (FParse::Command(&Cmd, TEXT("PRESENCE")))
+					{
+						// Takes a user id/name of a non-friend user for the sole usage of querying out
+						// Pass nothing if the platform doesn't support it
+						(new FTestPresenceInterface(SubName))->Test(InWorld, FParse::Token(Cmd, false));
+						bWasHandled = true;
+					}
 					else if (FParse::Command(&Cmd, TEXT("VOICE")))
 					{
 						// This class deletes itself once done
diff --git a/Engine/Plugins/Runtime/PacketHandlers/CompressionComponents/Oodle/Source/OodleHandlerComponent/Public/OodleHandlerComponent.h b/Engine/Plugins/Runtime/PacketHandlers/CompressionComponents/Oodle/Source/OodleHandlerComponent/Public/OodleHandlerComponent.h
--- a/Engine/Plugins/Runtime/PacketHandlers/CompressionComponents/Oodle/Source/OodleHandlerComponent/Public/OodleHandlerComponent.h
+++ b/Engine/Plugins/Runtime/PacketHandlers/CompressionComponents/Oodle/Source/OodleHandlerComponent/Public/OodleHandlerComponent.h
@@ -23,7 +23,7 @@ DECLARE_LOG_CATEGORY_EXTERN(OodleHandlerComponentLog, Log, All);
 #if HAS_OODLE_SDK
 
 #if UE4_OODLE_VER >= 200
-#include <oodle2.h>
+#include "oodle2.h"
 #else
 #include "oodle.h"
 #endif
diff --git a/Engine/Source/Editor/BlueprintGraph/Private/K2Node_ActorBoundEvent.cpp b/Engine/Source/Editor/BlueprintGraph/Private/K2Node_ActorBoundEvent.cpp
--- a/Engine/Source/Editor/BlueprintGraph/Private/K2Node_ActorBoundEvent.cpp
+++ b/Engine/Source/Editor/BlueprintGraph/Private/K2Node_ActorBoundEvent.cpp
@@ -81,7 +81,7 @@ void UK2Node_ActorBoundEvent::DestroyNode()
 	if (EventOwner)
 	{
 		// If we have an event owner, remove the delegate referencing this event, if any
-		const ULevel* TargetLevel = Cast<ULevel>(EventOwner->GetOuter());
+		const ULevel* TargetLevel = EventOwner->GetLevel();
 		if (TargetLevel)
 		{
 			ALevelScriptActor* LSA = TargetLevel->GetLevelScriptActor();
diff --git a/Engine/Source/Editor/DetailCustomizations/Private/ActorDetails.cpp b/Engine/Source/Editor/DetailCustomizations/Private/ActorDetails.cpp
--- a/Engine/Source/Editor/DetailCustomizations/Private/ActorDetails.cpp
+++ b/Engine/Source/Editor/DetailCustomizations/Private/ActorDetails.cpp
@@ -96,7 +96,7 @@ void FActorDetails::CustomizeDetails( IDetailLayoutBuilder& DetailLayout )
 				SelectedActors.Add( Actor );
 
 				// Record the level that contains this actor and increment it's actor count
-				ULevel* Level = Actor->GetTypedOuter<ULevel>();
+				ULevel* Level = Actor->GetLevel();
 				if (Level != NULL)
 				{
 					int32& ActorCountForThisLevel = ActorsPerLevelCount.FindOrAdd(Level);
diff --git a/Engine/Source/Editor/UnrealEd/Classes/Settings/LevelEditorPlaySettings.h b/Engine/Source/Editor/UnrealEd/Classes/Settings/LevelEditorPlaySettings.h
--- a/Engine/Source/Editor/UnrealEd/Classes/Settings/LevelEditorPlaySettings.h
+++ b/Engine/Source/Editor/UnrealEd/Classes/Settings/LevelEditorPlaySettings.h
@@ -304,7 +304,7 @@ private:
 	int32 PlayNumberOfClients;
 
 	/** What port used by the server for simple networking */
-	UPROPERTY(config, EditAnywhere, Category = MultiplayerOptions)
+	UPROPERTY(config, EditAnywhere, Category = MultiplayerOptions, meta=(ClampMin="1", UIMin="1", ClampMax="65535"))
 	uint16 ServerPort;
 
 	/** Width to use when spawning additional windows. */
diff --git a/Engine/Source/Editor/UnrealEd/Private/AssetSelection.cpp b/Engine/Source/Editor/UnrealEd/Private/AssetSelection.cpp
--- a/Engine/Source/Editor/UnrealEd/Private/AssetSelection.cpp
+++ b/Engine/Source/Editor/UnrealEd/Private/AssetSelection.cpp
@@ -121,7 +121,8 @@ namespace AssetSelectionUtils
 
 					if( ActorInfo.bAllSelectedActorsBelongToCurrentLevel )
 					{
-						if( !CurrentActor->GetOuter()->IsA(ULevel::StaticClass()) || !CurrentActor->GetLevel()->IsCurrentLevel() )
+						const ULevel* ActorLevel = CurrentActor->GetLevel();
+						if( !ActorLevel || !ActorLevel->IsCurrentLevel() )
 						{
 							ActorInfo.bAllSelectedActorsBelongToCurrentLevel = false;
 						}
@@ -159,7 +160,7 @@ namespace AssetSelectionUtils
 
 					if ( ActorInfo.bSelectedActorsBelongToSameLevel )
 					{
-						ULevel* ActorLevel = CurrentActor->GetOuter()->IsA(ULevel::StaticClass()) ? CurrentActor->GetLevel() : NULL;
+						ULevel* ActorLevel = CurrentActor->GetLevel();
 						if ( !ActorInfo.SharedLevel )
 						{
 							// This is the first selected actor we've encountered.
diff --git a/Engine/Source/Editor/UnrealEd/Private/EditorEngine.cpp b/Engine/Source/Editor/UnrealEd/Private/EditorEngine.cpp
--- a/Engine/Source/Editor/UnrealEd/Private/EditorEngine.cpp
+++ b/Engine/Source/Editor/UnrealEd/Private/EditorEngine.cpp
@@ -6810,10 +6810,15 @@ void UEditorEngine::UpdateAutoLoadProject()
 	IFileManager::Get().Delete(*AutoLoadInProgressFilename, bRequireExists, bEvenIfReadOnly, bQuiet);
 }
 
-FORCEINLINE bool NetworkRemapPath_local(FWorldContext& Context, FString& Str, bool bReading)
+FORCEINLINE bool NetworkRemapPath_local(FWorldContext& Context, FString& Str, bool bReading, bool bIsReplay)
 {
 	if (bReading)
 	{
+		if (bIsReplay && Context.World() && Context.World()->RemapCompiledScriptActor(Str))
+		{
+			return true;
+		}
+		
 		if (FPackageName::IsShortPackageName(Str))
 		{
 			return false;
@@ -6821,7 +6826,7 @@ FORCEINLINE bool NetworkRemapPath_local(FWorldContext& Context, FString& Str, bo
 
 		// First strip any source prefix, then add the appropriate prefix for this context
 		FSoftObjectPath Path = UWorld::RemovePIEPrefix(Str);
-			
+		
 		Path.FixupForPIE();
 		FString Remapped = Path.ToString();
 		if (!Remapped.Equals(Str, ESearchCase::CaseSensitive))
@@ -6850,14 +6855,16 @@ bool UEditorEngine::NetworkRemapPath(UNetDriver* Driver, FString& Str, bool bRea
 		return false;
 	}
 
+	// Pretty sure there's no case where you can't have a world by this point.
+	bool bIsAReplay = (Driver->GetWorld()) ? (Driver->GetWorld()->DemoNetDriver != NULL) : false;
 	FWorldContext& Context = GetWorldContextFromWorldChecked(Driver->GetWorld());
-	return NetworkRemapPath_local(Context, Str, bReading);
+	return NetworkRemapPath_local(Context, Str, bReading, bIsAReplay);
 }
 
 bool UEditorEngine::NetworkRemapPath( UPendingNetGame *PendingNetGame, FString& Str, bool bReading)
 {
 	FWorldContext& Context = GetWorldContextFromPendingNetGameChecked(PendingNetGame);
-	return NetworkRemapPath_local(Context, Str, bReading);
+	return NetworkRemapPath_local(Context, Str, bReading, PendingNetGame->DemoNetDriver != NULL);
 }
 
 void UEditorEngine::VerifyLoadMapWorldCleanup()
diff --git a/Engine/Source/Editor/UnrealEd/Private/EditorServer.cpp b/Engine/Source/Editor/UnrealEd/Private/EditorServer.cpp
--- a/Engine/Source/Editor/UnrealEd/Private/EditorServer.cpp
+++ b/Engine/Source/Editor/UnrealEd/Private/EditorServer.cpp
@@ -4471,20 +4471,7 @@ AActor* UEditorEngine::SelectNamedActor(const TCHAR* TargetActorName)
  */
 static bool IsInALevel(UObject* Obj)
 {
-	UObject* Outer = Obj->GetOuter();
-
-	// Keep looping while we walk up Outer chain.
-	while(Outer)
-	{
-		if(Outer->IsA(ULevel::StaticClass()))
-		{
-			return true;
-		}
-
-		Outer = Outer->GetOuter();
-	}
-
-	return false;
+	return Obj->GetTypedOuter<ULevel>() != nullptr;
 }
 
 
diff --git a/Engine/Source/Runtime/Core/Private/HAL/MallocAnsi.cpp b/Engine/Source/Runtime/Core/Private/HAL/MallocAnsi.cpp
--- a/Engine/Source/Runtime/Core/Private/HAL/MallocAnsi.cpp
+++ b/Engine/Source/Runtime/Core/Private/HAL/MallocAnsi.cpp
@@ -34,6 +34,15 @@ FMallocAnsi::FMallocAnsi()
 void* FMallocAnsi::Malloc( SIZE_T Size, uint32 Alignment )
 {
 	IncrementTotalMallocCalls();
+
+#if !UE_BUILD_SHIPPING
+	if (MaxSingleAlloc != 0 && Size > MaxSingleAlloc)
+	{
+		FPlatformMemory::OnOutOfMemory(Size, Alignment);
+		return nullptr;
+	}
+#endif
+
 	Alignment = FMath::Max(Size >= 16 ? (uint32)16 : (uint32)8, Alignment);
 
 #if USE_ALIGNED_MALLOC
@@ -62,6 +71,15 @@ void* FMallocAnsi::Malloc( SIZE_T Size, uint32 Alignment )
 void* FMallocAnsi::Realloc( void* Ptr, SIZE_T NewSize, uint32 Alignment )
 {
 	IncrementTotalReallocCalls();
+
+#if !UE_BUILD_SHIPPING
+	if (MaxSingleAlloc != 0 && NewSize > MaxSingleAlloc)
+	{
+		FPlatformMemory::OnOutOfMemory(NewSize, Alignment);
+		return nullptr;
+	}
+#endif
+
 	void* Result;
 	Alignment = FMath::Max(NewSize >= 16 ? (uint32)16 : (uint32)8, Alignment);
 
diff --git a/Engine/Source/Runtime/Core/Private/HAL/MallocTBB.cpp b/Engine/Source/Runtime/Core/Private/HAL/MallocTBB.cpp
--- a/Engine/Source/Runtime/Core/Private/HAL/MallocTBB.cpp
+++ b/Engine/Source/Runtime/Core/Private/HAL/MallocTBB.cpp
@@ -37,6 +37,14 @@ void* FMallocTBB::Malloc( SIZE_T Size, uint32 Alignment )
 {
 	IncrementTotalMallocCalls();
 
+#if !UE_BUILD_SHIPPING
+	if (MaxSingleAlloc != 0 && Size > MaxSingleAlloc)
+	{
+		FPlatformMemory::OnOutOfMemory(Size, Alignment);
+		return nullptr;
+	}
+#endif
+
 	MEM_TIME(MemTime -= FPlatformTime::Seconds());
 
 	void* NewPtr = NULL;
@@ -77,6 +85,14 @@ void* FMallocTBB::Realloc( void* Ptr, SIZE_T NewSize, uint32 Alignment )
 {
 	IncrementTotalReallocCalls();
 
+#if !UE_BUILD_SHIPPING
+	if (MaxSingleAlloc != 0 && NewSize > MaxSingleAlloc)
+	{
+		FPlatformMemory::OnOutOfMemory(NewSize, Alignment);
+		return nullptr;
+	}
+#endif
+
 	MEM_TIME(MemTime -= FPlatformTime::Seconds())
 #if UE_BUILD_DEBUG || UE_BUILD_DEVELOPMENT
 	SIZE_T OldSize = 0;
diff --git a/Engine/Source/Runtime/Core/Private/HAL/MemoryBase.cpp b/Engine/Source/Runtime/Core/Private/HAL/MemoryBase.cpp
--- a/Engine/Source/Runtime/Core/Private/HAL/MemoryBase.cpp
+++ b/Engine/Source/Runtime/Core/Private/HAL/MemoryBase.cpp
@@ -14,6 +14,10 @@ uint32 FMalloc::TotalMallocCalls = 0;
 uint32 FMalloc::TotalFreeCalls = 0;
 uint32 FMalloc::TotalReallocCalls = 0;
 
+#if !UE_BUILD_SHIPPING
+uint64 FMalloc::MaxSingleAlloc = 0;
+#endif
+
 struct FCurrentFrameCalls
 {
 	uint32 LastMallocCalls;
diff --git a/Engine/Source/Runtime/Core/Private/Serialization/Archive.cpp b/Engine/Source/Runtime/Core/Private/Serialization/Archive.cpp
--- a/Engine/Source/Runtime/Core/Private/Serialization/Archive.cpp
+++ b/Engine/Source/Runtime/Core/Private/Serialization/Archive.cpp
@@ -147,6 +147,7 @@ void FArchive::Reset()
 	ArMaxSerializeSize					= 0;
 	ArIsFilterEditorOnly				= false;
 	ArIsSaveGame						= false;
+	ArIsNetArchive						= false;
 	ArCustomPropertyList				= nullptr;
 	ArUseCustomPropertyList				= false;
 	CookingTargetPlatform = nullptr;
@@ -198,6 +199,7 @@ void FArchive::CopyTrivialFArchiveStatusMembers(const FArchive& ArchiveToCopy)
 	ArMaxSerializeSize                   = ArchiveToCopy.ArMaxSerializeSize;
 	ArIsFilterEditorOnly                 = ArchiveToCopy.ArIsFilterEditorOnly;
 	ArIsSaveGame                         = ArchiveToCopy.ArIsSaveGame;
+	ArIsNetArchive                       = ArchiveToCopy.ArIsNetArchive;
 	ArCustomPropertyList				 = ArchiveToCopy.ArCustomPropertyList;
 	ArUseCustomPropertyList				 = ArchiveToCopy.ArUseCustomPropertyList;
 	CookingTargetPlatform                = ArchiveToCopy.CookingTargetPlatform;
diff --git a/Engine/Source/Runtime/Core/Private/Serialization/BitReader.cpp b/Engine/Source/Runtime/Core/Private/Serialization/BitReader.cpp
--- a/Engine/Source/Runtime/Core/Private/Serialization/BitReader.cpp
+++ b/Engine/Source/Runtime/Core/Private/Serialization/BitReader.cpp
@@ -123,6 +123,9 @@ FBitReader::FBitReader(uint8* Src, int64 CountBits)
 
 	ArIsPersistent = ArIsLoading = 1;
 
+	// This class is exclusively used by the netcode
+	ArIsNetArchive = true;
+
 	if (Src != nullptr)
 	{
 		FMemory::Memcpy(Buffer.GetData(), Src, (CountBits + 7) >> 3);
diff --git a/Engine/Source/Runtime/Core/Private/Serialization/BitWriter.cpp b/Engine/Source/Runtime/Core/Private/Serialization/BitWriter.cpp
--- a/Engine/Source/Runtime/Core/Private/Serialization/BitWriter.cpp
+++ b/Engine/Source/Runtime/Core/Private/Serialization/BitWriter.cpp
@@ -30,6 +30,9 @@ FBitWriter::FBitWriter( int64 InMaxBits, bool InAllowResize /*=false*/ )
 	AllowResize = InAllowResize;
 	FMemory::Memzero(Buffer.GetData(), Buffer.Num());
 	ArIsPersistent = ArIsSaving = 1;
+
+	// This class is exclusively used by the netcode
+	ArIsNetArchive = true;
 }
 
 /**
@@ -41,6 +44,8 @@ FBitWriter::FBitWriter(void)
 	, AllowResize(false)
 	, bAllowOverflow(false)
 {
+	// This class is exclusively used by the netcode
+	ArIsNetArchive = true;
 }
 
 /**
diff --git a/Engine/Source/Runtime/Core/Private/UObject/UnrealNames.cpp b/Engine/Source/Runtime/Core/Private/UObject/UnrealNames.cpp
--- a/Engine/Source/Runtime/Core/Private/UObject/UnrealNames.cpp
+++ b/Engine/Source/Runtime/Core/Private/UObject/UnrealNames.cpp
@@ -526,8 +526,13 @@ int32 FName::Compare( const FName& Other ) const
 		const FNameEntry* const ThisEntry = GetComparisonNameEntry();
 		const FNameEntry* const OtherEntry = Other.GetComparisonNameEntry();
 
+		// If one or both entries return an invalid name entry, the comparison fails - fallback to comparing the index
+		if (ThisEntry == nullptr || OtherEntry == nullptr)
+		{
+			return GetComparisonIndexFast() - Other.GetComparisonIndexFast();
+		}
 		// Ansi/Wide mismatch, convert to wide
-		if( ThisEntry->IsWide() != OtherEntry->IsWide() )
+		else if( ThisEntry->IsWide() != OtherEntry->IsWide() )
 		{
 			return FCStringWide::Stricmp(	ThisEntry->IsWide() ? ThisEntry->GetWideName() : StringCast<WIDECHAR>(ThisEntry->GetAnsiName()).Get(),
 								OtherEntry->IsWide() ? OtherEntry->GetWideName() : StringCast<WIDECHAR>(OtherEntry->GetAnsiName()).Get() );
@@ -916,8 +921,11 @@ FString FName::ToString() const
 {
 	if (GetNumber() == NAME_NO_NUMBER_INTERNAL)
 	{
-		// Avoids some extra allocations in non-number case
-		return GetDisplayNameEntry()->GetPlainNameString();
+		if (const FNameEntry* const DisplayEntry = GetDisplayNameEntry())
+		{
+			// Avoids some extra allocations in non-number case
+			return DisplayEntry->GetPlainNameString();
+		}
 	}
 	
 	FString Out;	
@@ -929,7 +937,12 @@ void FName::ToString(FString& Out) const
 {
 	// A version of ToString that saves at least one string copy
 	const FNameEntry* const NameEntry = GetDisplayNameEntry();
-	if (GetNumber() == NAME_NO_NUMBER_INTERNAL)
+
+	if (NameEntry == nullptr)
+	{
+		Out = TEXT("*INVALID*");
+	}
+	else if (GetNumber() == NAME_NO_NUMBER_INTERNAL)
 	{
 		Out.Empty(NameEntry->GetNameLength());
 		NameEntry->AppendNameToString(Out);
@@ -947,11 +960,19 @@ void FName::ToString(FString& Out) const
 void FName::AppendString(FString& Out) const
 {
 	const FNameEntry* const NameEntry = GetDisplayNameEntry();
-	NameEntry->AppendNameToString( Out );
-	if (GetNumber() != NAME_NO_NUMBER_INTERNAL)
+
+	if (NameEntry == nullptr)
 	{
-		Out += TEXT("_");
-		Out.AppendInt(NAME_INTERNAL_TO_EXTERNAL(GetNumber()));
+		Out += TEXT("*INVALID*");
+	}
+	else
+	{
+		NameEntry->AppendNameToString( Out );
+		if (GetNumber() != NAME_NO_NUMBER_INTERNAL)
+		{
+			Out += TEXT("_");
+			Out.AppendInt(NAME_INTERNAL_TO_EXTERNAL(GetNumber()));
+		}
 	}
 }
 
diff --git a/Engine/Source/Runtime/Core/Public/Containers/Array.h b/Engine/Source/Runtime/Core/Public/Containers/Array.h
--- a/Engine/Source/Runtime/Core/Public/Containers/Array.h
+++ b/Engine/Source/Runtime/Core/Public/Containers/Array.h
@@ -1068,37 +1068,54 @@ public:
 	friend FArchive& operator<<(FArchive& Ar, TArray& A)
 	{
 		A.CountBytes(Ar);
-		if (sizeof(ElementType) == 1)
+
+		// For net archives, limit serialization to 16MB, to protect against excessive allocation
+		const int32 MaxNetArraySerialize = (16 * 1024 * 1024) / sizeof(ElementType);
+		int32 SerializeNum = (Ar.IsLoading() ? 0 : A.ArrayNum);
+
+		Ar << SerializeNum;
+
+		check(SerializeNum >= 0);
+
+		if (!Ar.IsError() && SerializeNum >= 0 && ensure(!Ar.IsNetArchive() || SerializeNum <= MaxNetArraySerialize))
 		{
-			// Serialize simple bytes which require no construction or destruction.
-			Ar << A.ArrayNum;
-			check(A.ArrayNum >= 0);
-			if ((A.ArrayNum || A.ArrayMax) && Ar.IsLoading())
+			if (sizeof(ElementType) == 1)
 			{
-				A.ResizeForCopy(A.ArrayNum, A.ArrayMax);
+				A.ArrayNum = SerializeNum;
+
+				// Serialize simple bytes which require no construction or destruction.
+				if ((A.ArrayNum || A.ArrayMax) && Ar.IsLoading())
+				{
+					A.ResizeForCopy(A.ArrayNum, A.ArrayMax);
+				}
+
+				Ar.Serialize(A.GetData(), A.Num());
 			}
-			Ar.Serialize(A.GetData(), A.Num());
-		}
-		else if (Ar.IsLoading())
-		{
-			// Load array.
-			int32 NewNum = 0;
-			Ar << NewNum;
-			A.Empty(NewNum);
-			for (int32 i = 0; i < NewNum; i++)
+			else if (Ar.IsLoading())
+			{
+				// Required for resetting ArrayNum
+				A.Empty(SerializeNum);
+
+				for (int32 i=0; i<SerializeNum; i++)
+				{
+					Ar << *::new(A) ElementType;
+				}
+			}
+			else
 			{
-				Ar << *::new(A)ElementType;
+				A.ArrayNum = SerializeNum;
+
+				for (int32 i=0; i<A.ArrayNum; i++)
+				{
+					Ar << A[i];
+				}
 			}
 		}
 		else
 		{
-			// Save array.
-			Ar << A.ArrayNum;
-			for (int32 i = 0; i < A.ArrayNum; i++)
-			{
-				Ar << A[i];
-			}
+			Ar.SetError();
 		}
+
 		return Ar;
 	}
 
diff --git a/Engine/Source/Runtime/Core/Public/HAL/MemoryBase.h b/Engine/Source/Runtime/Core/Public/HAL/MemoryBase.h
--- a/Engine/Source/Runtime/Core/Public/HAL/MemoryBase.h
+++ b/Engine/Source/Runtime/Core/Public/HAL/MemoryBase.h
@@ -214,6 +214,12 @@ protected:
 	static uint32 TotalFreeCalls;
 	/** Total number of calls Malloc, if implemented by derived class. */
 	static uint32 TotalReallocCalls;
+
+#if !UE_BUILD_SHIPPING
+public:
+	/** Limits the maximum single allocation, to this many bytes, for debugging */
+	static uint64 MaxSingleAlloc;
+#endif
 };
 
 
diff --git a/Engine/Source/Runtime/Core/Public/UObject/NameTypes.h b/Engine/Source/Runtime/Core/Public/UObject/NameTypes.h
--- a/Engine/Source/Runtime/Core/Public/UObject/NameTypes.h
+++ b/Engine/Source/Runtime/Core/Public/UObject/NameTypes.h
@@ -593,7 +593,12 @@ public:
 	/** Returns the pure name string without any trailing numbers */
 	FString GetPlainNameString() const
 	{
-		return GetDisplayNameEntry()->GetPlainNameString();
+		if (const FNameEntry* CurEntry = GetDisplayNameEntry())
+		{
+			return CurEntry->GetPlainNameString();
+		}
+
+		return TEXT("*INVALID*");
 	}
 
 	/**
@@ -601,7 +606,12 @@ public:
 	 */
 	FORCEINLINE ANSICHAR const* GetPlainANSIString() const
 	{
-		return GetDisplayNameEntry()->GetAnsiName();
+		if (const FNameEntry* CurEntry = GetDisplayNameEntry())
+		{
+			return CurEntry->GetAnsiName();
+		}
+
+		return "*INVALID*";
 	}
 
 	/**
@@ -609,7 +619,12 @@ public:
 	 */
 	FORCEINLINE WIDECHAR const* GetPlainWIDEString() const
 	{
-		return GetDisplayNameEntry()->GetWideName();
+		if (const FNameEntry* CurEntry = GetDisplayNameEntry())
+		{
+			return CurEntry->GetWideName();
+		}
+
+		return L"*INVALID*";
 	}
 
 	const FNameEntry* GetComparisonNameEntry() const;
@@ -935,33 +950,36 @@ public:
 		}
 
 		// Find name entry associated with this FName.
+		bool bAreNamesMatching = false;
 		const FNameEntry* const Entry = GetComparisonNameEntry();
 
-		// Temporary buffer to hold split name in case passed in name is of Name_Number format.
-		WIDECHAR TempBuffer[NAME_SIZE];
-		int32 InNumber = NAME_NO_NUMBER_INTERNAL;
-		int32 TempNumber = NAME_NO_NUMBER_INTERNAL;
-
-		// Check whether we need to split the passed in string into name and number portion.
-		auto WideOther = StringCast<WIDECHAR>(Other);
-		const WIDECHAR* WideOtherPtr = WideOther.Get();
-		if (SplitNameWithCheck(WideOtherPtr, TempBuffer, ARRAY_COUNT(TempBuffer), TempNumber))
+		if (Entry != nullptr)
 		{
-			WideOtherPtr = TempBuffer;
-			InNumber = NAME_EXTERNAL_TO_INTERNAL(TempNumber);
-		}
-
-		// Report a match if both the number and string portion match.
-		bool bAreNamesMatching = false;
-		if (InNumber == GetNumber())
-		{
-			if (Entry->IsWide())
+			// Temporary buffer to hold split name in case passed in name is of Name_Number format.
+			WIDECHAR TempBuffer[NAME_SIZE];
+			int32 InNumber = NAME_NO_NUMBER_INTERNAL;
+			int32 TempNumber = NAME_NO_NUMBER_INTERNAL;
+
+			// Check whether we need to split the passed in string into name and number portion.
+			auto WideOther = StringCast<WIDECHAR>(Other);
+			const WIDECHAR* WideOtherPtr = WideOther.Get();
+			if (SplitNameWithCheck(WideOtherPtr, TempBuffer, ARRAY_COUNT(TempBuffer), TempNumber))
 			{
-				bAreNamesMatching = !FPlatformString::Stricmp(WideOtherPtr, Entry->GetWideName());
+				WideOtherPtr = TempBuffer;
+				InNumber = NAME_EXTERNAL_TO_INTERNAL(TempNumber);
 			}
-			else
+
+			// Report a match if both the number and string portion match.
+			if (InNumber == GetNumber())
 			{
-				bAreNamesMatching = !FPlatformString::Stricmp(WideOtherPtr, Entry->GetAnsiName());
+				if (Entry->IsWide())
+				{
+					bAreNamesMatching = !FPlatformString::Stricmp(WideOtherPtr, Entry->GetWideName());
+				}
+				else
+				{
+					bAreNamesMatching = !FPlatformString::Stricmp(WideOtherPtr, Entry->GetAnsiName());
+				}
 			}
 		}
 
diff --git a/Engine/Source/Runtime/Engine/Classes/Engine/EngineTypes.h b/Engine/Source/Runtime/Engine/Classes/Engine/EngineTypes.h
--- a/Engine/Source/Runtime/Engine/Classes/Engine/EngineTypes.h
+++ b/Engine/Source/Runtime/Engine/Classes/Engine/EngineTypes.h
@@ -2719,16 +2719,17 @@ UENUM()
 enum ENetDormancy
 {
 	/** This actor can never go network dormant. */
-	DORM_Never,
+	DORM_Never UMETA(DisplayName = "Never"),
 	/** This actor can go dormant, but is not currently dormant. Game code will tell it when it go dormant. */
-	DORM_Awake,
+	DORM_Awake UMETA(DisplayName = "Awake"),
 	/** This actor wants to go fully dormant for all connections. */
-	DORM_DormantAll,
+	DORM_DormantAll UMETA(DisplayName = "Dormant All"),
 	/** This actor may want to go dormant for some connections, GetNetDormancy() will be called to find out which. */
-	DORM_DormantPartial,
+	DORM_DormantPartial UMETA(DisplayName = "Dormant Partial"),
 	/** This actor is initially dormant for all connection if it was placed in map. */
-	DORM_Initial,
-	DORN_MAX,
+	DORM_Initial UMETA(DisplayName = "Initial"),
+
+	DORM_MAX UMETA(Hidden),
 };
 
 
diff --git a/Engine/Source/Runtime/Engine/Classes/Engine/NetConnection.h b/Engine/Source/Runtime/Engine/Classes/Engine/NetConnection.h
--- a/Engine/Source/Runtime/Engine/Classes/Engine/NetConnection.h
+++ b/Engine/Source/Runtime/Engine/Classes/Engine/NetConnection.h
@@ -792,6 +792,8 @@ public:
 	/** Whether or not a client packet has been received - used serverside, to delay any packet sends */
 	FORCEINLINE bool HasReceivedClientPacket()
 	{
+		// The InternalAck and ServerConnection conditions, are only there to exclude demo's and clients from this check,
+		// so that the check is only performed on servers.
 		return !!InternalAck || Driver->ServerConnection != nullptr || InReliable[0] != InitInReliable;
 	}
 
diff --git a/Engine/Source/Runtime/Engine/Classes/Engine/NetDriver.h b/Engine/Source/Runtime/Engine/Classes/Engine/NetDriver.h
--- a/Engine/Source/Runtime/Engine/Classes/Engine/NetDriver.h
+++ b/Engine/Source/Runtime/Engine/Classes/Engine/NetDriver.h
@@ -432,6 +432,11 @@ public:
 	 */
 	TMap<FNetworkGUID, FActorDestructionInfo>	DestroyedStartupOrDormantActors;
 
+	/** The server adds an entry into this map for every startup actor that has been renamed, and will
+	 *  always map from current name to original name
+	 */
+	TMap<FName, FName>	RenamedStartupActors;
+
 	/** Maps FRepChangedPropertyTracker to active objects that are replicating properties */
 	TMap< TWeakObjectPtr< UObject >, TSharedPtr< FRepChangedPropertyTracker > >	RepChangedPropertyTrackerMap;
 	/** Used to invalidate properties marked "unchanged" in FRepChangedPropertyTracker's */
@@ -686,6 +691,9 @@ public:
 	/** Called when a spawned actor is destroyed. */
 	ENGINE_API virtual void NotifyActorDestroyed( AActor* Actor, bool IsSeamlessTravel=false );
 
+	/** Called when an actor is renamed. */
+	ENGINE_API virtual void NotifyActorRenamed(AActor* Actor, FName PreviousName);
+
 	ENGINE_API virtual void NotifyStreamingLevelUnload( ULevel* );
 
 	ENGINE_API virtual void NotifyActorLevelUnloaded( AActor* Actor );
diff --git a/Engine/Source/Runtime/Engine/Classes/Engine/World.h b/Engine/Source/Runtime/Engine/Classes/Engine/World.h
--- a/Engine/Source/Runtime/Engine/Classes/Engine/World.h
+++ b/Engine/Source/Runtime/Engine/Classes/Engine/World.h
@@ -3195,6 +3195,9 @@ public:
 	/** Rename this world such that it has the prefix on names for the given PIE Instance ID */
 	void RenameToPIEWorld(int32 PIEInstanceID);
 
+	/** Given a level script actor, modify the string such that it points to the correct instance of the object. For replays. */
+	bool RemapCompiledScriptActor(FString& Str) const;
+
 	/** Given a PackageName and a PIE Instance ID return the name of that Package when being run as a PIE world */
 	static FString ConvertToPIEPackageName(const FString& PackageName, int32 PIEInstanceID);
 
diff --git a/Engine/Source/Runtime/Engine/Classes/GameFramework/Actor.h b/Engine/Source/Runtime/Engine/Classes/GameFramework/Actor.h
--- a/Engine/Source/Runtime/Engine/Classes/GameFramework/Actor.h
+++ b/Engine/Source/Runtime/Engine/Classes/GameFramework/Actor.h
@@ -1461,6 +1461,7 @@ public:
 	virtual void BeginDestroy() override;
 	virtual bool IsReadyForFinishDestroy() override;
 	virtual bool Rename( const TCHAR* NewName=nullptr, UObject* NewOuter=nullptr, ERenameFlags Flags=REN_None ) override;
+	virtual void PostRename( UObject* OldOuter, const FName OldName ) override;
 	virtual bool CanBeInCluster() const override;
 	static void AddReferencedObjects(UObject* InThis, FReferenceCollector& Collector);
 #if WITH_EDITOR
