Added NetcodeUnitTest hook, for replacing the raw data for specific RPC parameters, on the fly.

diff --git a/Engine/Plugins/NetcodeUnitTest/NetcodeUnitTest/Source/NetcodeUnitTest/Private/Net/NUTUtilNet.cpp b/Engine/Plugins/NetcodeUnitTest/NetcodeUnitTest/Source/NetcodeUnitTest/Private/Net/NUTUtilNet.cpp
--- a/Engine/Plugins/NetcodeUnitTest/NetcodeUnitTest/Source/NetcodeUnitTest/Private/Net/NUTUtilNet.cpp
+++ b/Engine/Plugins/NetcodeUnitTest/NetcodeUnitTest/Source/NetcodeUnitTest/Private/Net/NUTUtilNet.cpp
@@ -1,6 +1,7 @@
 // Copyright 1998-2018 Epic Games, Inc. All Rights Reserved.
 
 #include "Net/NUTUtilNet.h"
+
 #include "UObject/CoreOnline.h"
 #include "GameFramework/OnlineReplStructs.h"
 #include "Engine/Engine.h"
@@ -10,6 +11,7 @@
 #include "Engine/LocalPlayer.h"
 #include "EngineUtils.h"
 #include "Net/DataChannel.h"
+#include "Net/RepLayout.h"
 
 #include "NUTUtil.h"
 #include "NUTUtilReflection.h"
@@ -25,6 +27,13 @@
 class FWorldTickHook;
 
 
+// Enable access to FRepLayout.Cmds
+IMPLEMENT_GET_PRIVATE_VAR(FRepLayout, Cmds, TArray<FRepLayoutCmd>);
+
+
+IMPLEMENT_INTRINSIC_CLASS(UNetPropertyHook, NETCODEUNITTEST_API, UProperty, COREUOBJECT_API, "/Script/NetcodeUnitTest", {});
+
+
 /** Active unit test worlds */
 TArray<UWorld*> UnitTestWorlds;
 
@@ -381,6 +390,81 @@ FScopedNetNameReplace::~FScopedNetNameReplace()
 	}
 }
 
+/**
+ * FScopedRPCParamReplace
+ */
+
+FScopedRPCParamReplace::FScopedRPCParamReplace(UMinimalClient* InMinClient, UFunction* InTargetRPC, FString InParamName,
+												const FOnNetSerializeItem& InSerializeHook)
+	: ParamRepCmd(nullptr)
+	, OriginalParam(nullptr)
+{
+	TSharedPtr<FRepLayout> FuncRepLayout;
+
+	if (UNetConnection* UnitConn = (InMinClient != nullptr ? InMinClient->GetConn() : nullptr))
+	{
+		if (UNetDriver* NetDriver = UnitConn->GetDriver())
+		{
+			FuncRepLayout = NetDriver->GetFunctionRepLayout(InTargetRPC);
+		}
+	}
+
+	if (FuncRepLayout.IsValid())
+	{
+		TArray<FRepLayoutCmd>& Cmds = GET_PRIVATE(FRepLayout, FuncRepLayout.Get(), Cmds);
+
+		for (FRepLayoutCmd& CurCmd : Cmds)
+		{
+			if (UProperty* CurProp = CurCmd.Property)
+			{
+				if (CurProp->GetName() == InParamName)
+				{
+					ParamRepCmd = &CurCmd;
+					OriginalParam = CurProp;
+
+					UNetPropertyHook* HookParam = GetMutableDefault<UNetPropertyHook>();
+
+					HookParam->SerializeHook = InSerializeHook;
+
+					ParamRepCmd->Property = HookParam;
+
+					break;
+				}
+			}
+		}
+	}
+}
+
+
+FScopedRPCParamReplace::~FScopedRPCParamReplace()
+{
+	GetMutableDefault<UNetPropertyHook>()->SerializeHook.Unbind();
+
+	if (ParamRepCmd != nullptr)
+	{
+		ParamRepCmd->Property = OriginalParam;
+	}
+
+	ParamRepCmd = nullptr;
+	OriginalParam = nullptr;
+}
+
+/**
+ * UNetPropertyHook
+ */
+
+bool UNetPropertyHook::NetSerializeItem(FArchive& Ar, UPackageMap* Map, void* Data, TArray<uint8>* MetaData) const
+{
+	bool bReturnVal = false;
+
+	if (SerializeHook.IsBound())
+	{
+		bReturnVal = SerializeHook.Execute(Ar, Map, Data, MetaData);
+	}
+
+	return bReturnVal;
+}
+
 
 /**
  * NUTNet
diff --git a/Engine/Plugins/NetcodeUnitTest/NetcodeUnitTest/Source/NetcodeUnitTest/Public/Net/NUTUtilNet.h b/Engine/Plugins/NetcodeUnitTest/NetcodeUnitTest/Source/NetcodeUnitTest/Public/Net/NUTUtilNet.h
--- a/Engine/Plugins/NetcodeUnitTest/NetcodeUnitTest/Source/NetcodeUnitTest/Public/Net/NUTUtilNet.h
+++ b/Engine/Plugins/NetcodeUnitTest/NetcodeUnitTest/Source/NetcodeUnitTest/Public/Net/NUTUtilNet.h
@@ -16,12 +16,14 @@ class APlayerController;
 class FInBunch;
 class FInternetAddr;
 class FOutBunch;
-class UMinimalClient;
+class FRepLayoutCmd;
 class UNetConnection;
 class UNetDriver;
+struct FUniqueNetIdRepl;
+
+class UMinimalClient;
 class UUnitTestChannel;
 class UUnitTestNetDriver;
-struct FUniqueNetIdRepl;
 class UClientUnitTest;
 
 
@@ -38,6 +40,17 @@ class UClientUnitTest;
 */
 DECLARE_DELEGATE_FourParams(FOnProcessNetEvent, AActor* /*Actor*/, UFunction* /*Function*/, void* /*Parameters*/, bool& /*bBlockRPC*/);
 
+/**
+ * Delegate used to control hooked serialization of RPC parameters
+ *
+ * @param Ar			The archive being serialized to
+ * @param Map			The PackageMap for the net connection
+ * @param Data			The raw property data
+ * @param MetaData		The meta data for the property
+ */
+DECLARE_DELEGATE_RetVal_FourParams(bool, FOnNetSerializeItem, FArchive& /*Ar*/, UPackageMap* /*Map*/, void* /*Data*/,
+									TArray<uint8>* /*MetaData*/);
+
 
 /**
  * Class for encapsulating ProcessEvent and general RPC hooks, implemented globally for each UWorld
@@ -309,6 +322,37 @@ private:
 	FDelegateHandle Handle;
 };
 
+/**
+ * Hooks the netcode serialization of a specific RPC parameter, in order to override the raw serialized param data with something else
+ */
+class NETCODEUNITTEST_API FScopedRPCParamReplace
+{
+public:
+	FScopedRPCParamReplace(UMinimalClient* InMinClient, UFunction* InTargetRPC, FString InParamName,
+							const FOnNetSerializeItem& InSerializeHook);
+
+	~FScopedRPCParamReplace();
+
+private:
+	FRepLayoutCmd*		ParamRepCmd;
+
+	UProperty*			OriginalParam;
+};
+
+/**
+ * Implements the actual hook for the above class
+ */
+class UNetPropertyHook : public UProperty
+{
+	DECLARE_CLASS_INTRINSIC(UNetPropertyHook, UProperty, 0, TEXT("/Script/NetcodeUnitTest"))
+
+private:
+	virtual bool NetSerializeItem(FArchive& Ar, UPackageMap* Map, void* Data, TArray<uint8>* MetaData /*=nullptr*/) const override;
+
+public:
+	FOnNetSerializeItem SerializeHook;
+};
+
 
 /**
  * Netcode based utility functions
