PR #4883: Fix compilation of unit test minimal client in shipping

diff --git a/Engine/Plugins/NetcodeUnitTest/NetcodeUnitTest/Source/NetcodeUnitTest/Private/MinimalClient.cpp b/Engine/Plugins/NetcodeUnitTest/NetcodeUnitTest/Source/NetcodeUnitTest/Private/MinimalClient.cpp
--- a/Engine/Plugins/NetcodeUnitTest/NetcodeUnitTest/Source/NetcodeUnitTest/Private/MinimalClient.cpp
+++ b/Engine/Plugins/NetcodeUnitTest/NetcodeUnitTest/Source/NetcodeUnitTest/Private/MinimalClient.cpp
@@ -292,7 +292,7 @@ bool UMinimalClient::SendRawBunch(FOutBunch& Bunch, bool bAllowPartial/*=false*/
 
 
 				SendBunches.Add(NewBunch);
-				NewBunch->DebugString = FString::Printf(TEXT("Partial[%d]: %s"), SendBunches.Num(), *Bunch.DebugString);
+				NewBunch->SetDebugString(FString::Printf(TEXT("Partial[%d]: %s"), SendBunches.Num(), *Bunch.GetDebugString()));
 			}
 
 			UNIT_LOG(, TEXT("SendRawBunch: Split oversized bunch (%i bits) into '%i' partial packets."), BunchNumBits,
