Added unit tests for WritePackedVector.
Clamp values after scaling to prevent too large numbers being passed to FMath::RoundToInt which would result in unexpected values being read.

diff --git a/Engine/Plugins/NetcodeUnitTest/NetcodeUnitTest/Source/NetcodeUnitTest/Classes/UnitTests/Engine/PackedVectorTest.h b/Engine/Plugins/NetcodeUnitTest/NetcodeUnitTest/Source/NetcodeUnitTest/Classes/UnitTests/Engine/PackedVectorTest.h
new file mode 100644
--- /dev/null
+++ b/Engine/Plugins/NetcodeUnitTest/NetcodeUnitTest/Source/NetcodeUnitTest/Classes/UnitTests/Engine/PackedVectorTest.h
@@ -0,0 +1,22 @@
+// Copyright 1998-2018 Epic Games, Inc. All Rights Reserved.
+
+#pragma once
+
+#include "CoreMinimal.h"
+#include "UObject/ObjectMacros.h"
+#include "UnitTest.h"
+
+#include "PackedVectorTest.generated.h"
+
+
+/**
+ * Basic unit test for WritePackedValue and ReadPackedValue
+ */
+UCLASS()
+class UPackedVectorTest : public UUnitTest
+{
+	GENERATED_UCLASS_BODY()
+
+private:
+	virtual bool ExecuteUnitTest() override;
+};
diff --git a/Engine/Plugins/NetcodeUnitTest/NetcodeUnitTest/Source/NetcodeUnitTest/Private/UnitTests/Engine/PackedVectorTest.cpp b/Engine/Plugins/NetcodeUnitTest/NetcodeUnitTest/Source/NetcodeUnitTest/Private/UnitTests/Engine/PackedVectorTest.cpp
new file mode 100644
--- /dev/null
+++ b/Engine/Plugins/NetcodeUnitTest/NetcodeUnitTest/Source/NetcodeUnitTest/Private/UnitTests/Engine/PackedVectorTest.cpp
@@ -0,0 +1,173 @@
+// Copyright 1998-2018 Epic Games, Inc. All Rights Reserved.
+
+#include "UnitTests/Engine/PackedVectorTest.h"
+#include "Engine/NetSerialization.h"
+#include "UnitTestEnvironment.h"
+#include <cmath>
+
+namespace PackedVectorTest
+{
+
+	static bool AlmostEqualUlps(const float A, const float B, const int MaxUlps)
+	{
+		union FFloatToInt { float F; int32 I; };
+		FFloatToInt IntA;
+		FFloatToInt IntB;
+
+		IntA.F = A;
+		IntB.F = B;
+
+		if ((IntA.I ^ IntB.I) < 0)
+		{
+			// For different signs we only allow +/- 0.0f
+			return ((IntA.I | IntB.I) & 0x7FFFFFFF) == 0;
+		}
+
+		const int Ulps = FMath::Abs(IntA.I - IntB.I);
+		return Ulps <= MaxUlps;
+	}
+
+}
+
+UPackedVectorTest::UPackedVectorTest(const FObjectInitializer& ObjectInitializer)
+: Super(ObjectInitializer)
+{
+	UnitTestName = TEXT("PackedVector");
+	UnitTestType = TEXT("Test");
+
+	UnitTestDate = FDateTime(2018, 9, 19);
+
+	ExpectedResult.Add(TEXT("ShooterGame"), EUnitTestVerification::VerifiedFixed);
+
+	UnitTestTimeout = 5;
+}
+
+bool UPackedVectorTest::ExecuteUnitTest()
+{
+	TMap<FString, bool> TestResults;
+
+	TestResults.Add(TEXT("Commencing Read/WritePackedVector tests. Only fails will be shown in log."), "Success" && true);
+
+	static const float Quantize10_Values[] =
+	{
+		0.0f,
+		-180817.42f,
+		47.11f,
+		-FMath::Exp2(25.0f), // overflow
+		INFINITY, // non-finite
+	};
+
+	static const float Quantize100_Values[] =
+	{
+		0.0f,
+		+180720.42f,
+		-19751216.0f,
+		FMath::Exp2(31.0f),
+		-INFINITY,
+	};
+
+	struct TestCase
+	{
+		int ScaleFactor;
+		int BitsPerComponent;
+
+		const float* TestValues;
+		int TestValueCount;
+
+		TFunction<bool(FVector&, FArchive&)> Reader;
+		TFunction<bool(FVector, FArchive&)> Writer;
+
+	};
+
+	TestCase TestCases[] =
+	{
+		{
+			10,
+			24,
+			Quantize10_Values,
+			sizeof(Quantize10_Values) / sizeof(Quantize10_Values[0]),
+			[](FVector& Value, FArchive& Ar) { return ReadPackedVector<10, 24>(Value, Ar); },
+			[](FVector Value, FArchive& Ar) { return WritePackedVector<10, 24>(Value, Ar); },
+		},
+
+		{
+			100,
+			30,
+			Quantize100_Values,
+			sizeof(Quantize100_Values) / sizeof(Quantize100_Values[0]),
+			[](FVector& Value, FArchive& Ar) { return ReadPackedVector<100, 30>(Value, Ar); },
+			[](FVector Value, FArchive& Ar) { return WritePackedVector<100, 30>(Value, Ar); },
+		},
+	};
+
+	constexpr bool bAllowResize = false;
+	FBitWriter Writer(128, bAllowResize);
+
+	for (size_t TestCaseIt = 0, TestCaseEndIt = sizeof(TestCases) / sizeof(TestCases[0]); TestCaseIt != TestCaseEndIt; ++TestCaseIt)
+	{
+		const TestCase& Test = TestCases[TestCaseIt];
+		for (size_t ValueIt = 0, ValueEndIt = Test.TestValueCount; ValueIt != ValueEndIt; ++ValueIt)
+		{
+			Writer.Reset();
+
+			const float ScalarValue = Test.TestValues[ValueIt];
+			const FVector WriteValue(ScalarValue);
+			FVector ReadValue;
+
+			const bool bOverflowOrNan = !Test.Writer(WriteValue, Writer);
+			bool LocalSuccess = !Writer.GetError();
+
+			if (LocalSuccess)
+			{
+				FBitReader Reader(Writer.GetData(), Writer.GetNumBits());
+
+				Test.Reader(ReadValue, Reader);
+				LocalSuccess &= !Reader.GetError();
+				LocalSuccess = (ReadValue.X == ReadValue.Y) && (ReadValue.X == ReadValue.Z);
+				if (LocalSuccess)
+				{
+					// At this point we should have similar values as the original ones, except for NaN and overflowed values
+					if (bOverflowOrNan)
+					{
+						if (WriteValue.ContainsNaN())
+							LocalSuccess &= (ReadValue == FVector::ZeroVector);
+						else
+						{
+							// Overflowed value. Should probably be close to range start or end [-2^(ComponentBitCount)/ScaleFactor, 2^(ComponentBitCount)/ScaleFactor]
+							const float MaxValue = FMath::Exp2(float(Test.BitsPerComponent))/Test.ScaleFactor;
+							LocalSuccess &= PackedVectorTest::AlmostEqualUlps(FMath::Abs(ReadValue.X), MaxValue, 1);
+						}
+					}
+					else
+					{
+						const float ValueDiff = FMath::Abs(ReadValue.X - WriteValue.X);
+						LocalSuccess &= ValueDiff < 2.0f / Test.ScaleFactor; // The diff test might need some adjustment
+					}
+				}
+			}
+
+			if (!LocalSuccess)
+			{
+				TestResults.Add(FString::Printf(TEXT("Read/WritePackedVector failed with scale %d, bit count %d and value %f. Got %f"), Test.ScaleFactor, Test.BitsPerComponent, ScalarValue, ReadValue.X), LocalSuccess);
+			}
+		}
+	}
+
+	// Verify the results
+	for (TMap<FString, bool>::TConstIterator It(TestResults); It; ++It)
+	{
+		UNIT_LOG(ELogType::StatusImportant, TEXT("Test '%s' returned: %s"), *It.Key(), (It.Value() ? TEXT("Success") : TEXT("FAIL")));
+
+		if (!It.Value())
+		{
+			VerificationState = EUnitTestVerification::VerifiedNeedsUpdate;
+		}
+	}
+
+	if (VerificationState == EUnitTestVerification::Unverified)
+	{
+		VerificationState = EUnitTestVerification::VerifiedFixed;
+	}
+
+	return true;
+}
diff --git a/Engine/Source/Runtime/Engine/Classes/Engine/NetSerialization.h b/Engine/Source/Runtime/Engine/Classes/Engine/NetSerialization.h
--- a/Engine/Source/Runtime/Engine/Classes/Engine/NetSerialization.h
+++ b/Engine/Source/Runtime/Engine/Classes/Engine/NetSerialization.h
@@ -1145,6 +1145,9 @@ bool WritePackedVector(FVector Value, FArchive& Ar)	// Note Value is intended to
 {
 	check(Ar.IsSaving());
 
+	// Scale vector by quant factor first
+	Value *= ScaleFactor;
+
 	// Nan Check
 	if( Value.ContainsNaN() )
 	{
@@ -1154,13 +1157,14 @@ bool WritePackedVector(FVector Value, FArchive& Ar)	// Note Value is intended to
 		return false;
 	}
 
-	// Scale vector by quant factor first
-	Value *= ScaleFactor;
+	// Some platforms have RoundToInt implementations that essentially reduces the allowed inputs to 2^31.
+	const FVector ClampedValue = ClampVector(Value, FVector(-1073741824.0f), FVector(1073741760.0f));
+	bool bClamp = ClampedValue != Value;
 
 	// Do basically FVector::SerializeCompressed
-	int32 IntX	= FMath::RoundToInt(Value.X);
-	int32 IntY	= FMath::RoundToInt(Value.Y);
-	int32 IntZ	= FMath::RoundToInt(Value.Z);
+	int32 IntX	= FMath::RoundToInt(ClampedValue.X);
+	int32 IntY	= FMath::RoundToInt(ClampedValue.Y);
+	int32 IntZ	= FMath::RoundToInt(ClampedValue.Z);
 			
 	uint32 Bits	= FMath::Clamp<uint32>( FMath::CeilLogTwo( 1 + FMath::Max3( FMath::Abs(IntX), FMath::Abs(IntY), FMath::Abs(IntZ) ) ), 1, MaxBitsPerComponent ) - 1;
 
@@ -1173,17 +1177,15 @@ bool WritePackedVector(FVector Value, FArchive& Ar)	// Note Value is intended to
 	uint32 DY	= IntY + Bias;
 	uint32 DZ	= IntZ + Bias;
 
-	bool clamp=false;
-	
-	if (DX >= Max) { clamp=true; DX = static_cast<int32>(DX) > 0 ? Max-1 : 0; }
-	if (DY >= Max) { clamp=true; DY = static_cast<int32>(DY) > 0 ? Max-1 : 0; }
-	if (DZ >= Max) { clamp=true; DZ = static_cast<int32>(DZ) > 0 ? Max-1 : 0; }
+	if (DX >= Max) { bClamp=true; DX = static_cast<int32>(DX) > 0 ? Max-1 : 0; }
+	if (DY >= Max) { bClamp=true; DY = static_cast<int32>(DY) > 0 ? Max-1 : 0; }
+	if (DZ >= Max) { bClamp=true; DZ = static_cast<int32>(DZ) > 0 ? Max-1 : 0; }
 	
 	Ar.SerializeInt( DX, Max );
 	Ar.SerializeInt( DY, Max );
 	Ar.SerializeInt( DZ, Max );
 
-	return !clamp;
+	return !bClamp;
 }
 
 template<uint32 ScaleFactor, int32 MaxBitsPerComponent>
