Convert all remaining "Developer" modules to "UncookedOnly", to preserve existing behavior.

#rnx

diff --git a/Engine/Plugins/AI/HTNPlanner/HTNPlanner.uplugin b/Engine/Plugins/AI/HTNPlanner/HTNPlanner.uplugin
--- a/Engine/Plugins/AI/HTNPlanner/HTNPlanner.uplugin
+++ b/Engine/Plugins/AI/HTNPlanner/HTNPlanner.uplugin
@@ -23,7 +23,7 @@
         },
         {
 			"Name" : "HTNTestSuite",
-			"Type" : "Developer",
+			"Type" : "UncookedOnly",
 			"LoadingPhase" : "PreDefault"
 		}
 	]
diff --git a/Engine/Plugins/Developer/ActorLayerUtilities/ActorLayerUtilities.uplugin b/Engine/Plugins/Developer/ActorLayerUtilities/ActorLayerUtilities.uplugin
--- a/Engine/Plugins/Developer/ActorLayerUtilities/ActorLayerUtilities.uplugin
+++ b/Engine/Plugins/Developer/ActorLayerUtilities/ActorLayerUtilities.uplugin
@@ -18,7 +18,7 @@
 	[
 		{
 			"Name" : "ActorLayerUtilities",
-			"Type" : "Developer",
+			"Type" : "UncookedOnly",
 			"LoadingPhase" : "Default"
 		},
 
diff --git a/Engine/Plugins/Developer/BlankPlugin/BlankPlugin.uplugin b/Engine/Plugins/Developer/BlankPlugin/BlankPlugin.uplugin
--- a/Engine/Plugins/Developer/BlankPlugin/BlankPlugin.uplugin
+++ b/Engine/Plugins/Developer/BlankPlugin/BlankPlugin.uplugin
@@ -18,7 +18,7 @@
 	[
 		{
 			"Name" : "BlankPlugin",
-			"Type" : "Developer",
+			"Type" : "Runtime",
 			"LoadingPhase" : "Default"
 		}
 	]
diff --git a/Engine/Plugins/Developer/UObjectPlugin/UObjectPlugin.uplugin b/Engine/Plugins/Developer/UObjectPlugin/UObjectPlugin.uplugin
--- a/Engine/Plugins/Developer/UObjectPlugin/UObjectPlugin.uplugin
+++ b/Engine/Plugins/Developer/UObjectPlugin/UObjectPlugin.uplugin
@@ -18,7 +18,7 @@
 	[
 		{
 			"Name" : "UObjectPlugin",
-			"Type" : "Developer",
+			"Type" : "Runtime",
 			"LoadingPhase" : "Default"
 		}
 	]
diff --git a/Engine/Plugins/Messaging/MessagingDebugger/MessagingDebugger.uplugin b/Engine/Plugins/Messaging/MessagingDebugger/MessagingDebugger.uplugin
--- a/Engine/Plugins/Messaging/MessagingDebugger/MessagingDebugger.uplugin
+++ b/Engine/Plugins/Messaging/MessagingDebugger/MessagingDebugger.uplugin
@@ -15,7 +15,7 @@
 	[
 		{
 			"Name" : "MessagingDebugger",
-			"Type" : "Developer"
+			"Type" : "UncookedOnly"
 		}
 	],
 
diff --git a/Engine/Plugins/NetcodeUnitTest/NetcodeUnitTest/NetcodeUnitTest.uplugin b/Engine/Plugins/NetcodeUnitTest/NetcodeUnitTest/NetcodeUnitTest.uplugin
--- a/Engine/Plugins/NetcodeUnitTest/NetcodeUnitTest/NetcodeUnitTest.uplugin
+++ b/Engine/Plugins/NetcodeUnitTest/NetcodeUnitTest/NetcodeUnitTest.uplugin
@@ -18,7 +18,7 @@
 	[
 		{
 			"Name" : "NetcodeUnitTest",
-			"Type" : "Developer",
+			"Type" : "UncookedOnly",
 			"LoadingPhase" : "Default",
 			"WhitelistPlatforms" :
 			[
diff --git a/Engine/Plugins/Runtime/AnimDistanceMatching/AnimDistanceMatching.uplugin b/Engine/Plugins/Runtime/AnimDistanceMatching/AnimDistanceMatching.uplugin
--- a/Engine/Plugins/Runtime/AnimDistanceMatching/AnimDistanceMatching.uplugin
+++ b/Engine/Plugins/Runtime/AnimDistanceMatching/AnimDistanceMatching.uplugin
@@ -18,7 +18,7 @@
 	[
 		{
 			"Name" : "AnimDistanceMatching",
-			"Type" : "Developer",
+			"Type" : "UncookedOnly",
 			"LoadingPhase" : "Default"
 		}
 	]
diff --git a/Engine/Plugins/Tests/RuntimeTests/RuntimeTests.uplugin b/Engine/Plugins/Tests/RuntimeTests/RuntimeTests.uplugin
--- a/Engine/Plugins/Tests/RuntimeTests/RuntimeTests.uplugin
+++ b/Engine/Plugins/Tests/RuntimeTests/RuntimeTests.uplugin
@@ -17,7 +17,7 @@
 	"Modules": [
 		{
 			"Name": "RuntimeTests",
-			"Type": "Developer",
+			"Type": "UncookedOnly",
 			"LoadingPhase": "PreDefault"
 		}
 	]
