CIS Fix

diff --git a/Engine/Plugins/NetcodeUnitTest/NetcodeUnitTest/Source/NetcodeUnitTest/Public/NUTUtilReflection.h b/Engine/Plugins/NetcodeUnitTest/NetcodeUnitTest/Source/NetcodeUnitTest/Public/NUTUtilReflection.h
--- a/Engine/Plugins/NetcodeUnitTest/NetcodeUnitTest/Source/NetcodeUnitTest/Public/NUTUtilReflection.h
+++ b/Engine/Plugins/NetcodeUnitTest/NetcodeUnitTest/Source/NetcodeUnitTest/Public/NUTUtilReflection.h
@@ -682,7 +682,7 @@ private:
 	 */
 	FORCEINLINE bool CanCastArray() const
 	{
-		return !bIsError && BaseAddress != NULL && FieldInstance != NULL && FieldInstance.IsA(FArrayProperty::StaticClass()) &&
+		return !bIsError && BaseAddress != NULL && FieldInstance.IsValid() && FieldInstance.IsA(FArrayProperty::StaticClass()) &&
 				FieldAddress != NULL && bVerifiedFieldType && !bSetArrayElement;
 	}
 
@@ -693,7 +693,7 @@ private:
 	 */
 	FORCEINLINE bool CanCastStruct() const
 	{
-		return !bIsError && BaseAddress != NULL && FieldInstance != NULL && FieldInstance.IsA(UStruct::StaticClass()) &&
+		return !bIsError && BaseAddress != NULL && FieldInstance.IsValid() && FieldInstance.IsA(UStruct::StaticClass()) &&
 				FieldAddress != NULL && bVerifiedFieldType;
 	}
 
