Fixing nativized Blueprint code compilation errors

diff --git a/Engine/Plugins/NetcodeUnitTest/NetcodeUnitTest/Source/NetcodeUnitTest/Private/NUTUtilReflection.cpp b/Engine/Plugins/NetcodeUnitTest/NetcodeUnitTest/Source/NetcodeUnitTest/Private/NUTUtilReflection.cpp
--- a/Engine/Plugins/NetcodeUnitTest/NetcodeUnitTest/Source/NetcodeUnitTest/Private/NUTUtilReflection.cpp
+++ b/Engine/Plugins/NetcodeUnitTest/NetcodeUnitTest/Source/NetcodeUnitTest/Private/NUTUtilReflection.cpp
@@ -622,7 +622,7 @@ void FVMReflection::ProcessStructProperty()
 {
 	if (!IsPropertyArray() || (bVerifiedFieldType && bSetArrayElement))
 	{
-		FStructProperty* StructProp = CastField<FStructProperty>(FieldInstance);
+		FStructProperty* StructProp = FieldInstance.Get<FStructProperty>();
 
 		if (StructProp != nullptr && FieldAddress != nullptr)
 		{
@@ -659,7 +659,7 @@ InType* FVMReflection::GetWritableCast(const TCHAR* InTypeStr, bool bDoingUpCast
 			ReturnVal = (InType*)FieldAddress;
 		}
 #if UENUM_REFL
-		else if (FieldInstance.IsA<FEnumProperty>() && CastField<FEnumProperty>(FieldInstance)->GetUnderlyingProperty()->IsA(InTypeClass::StaticClass()))
+		else if (FieldInstance.IsA<FEnumProperty>() && FieldInstance.Get<FEnumProperty>()->GetUnderlyingProperty()->IsA(InTypeClass::StaticClass()))
 		{
 			ReturnVal = (InType*)FieldAddress;
 		}
@@ -737,7 +737,7 @@ InType FVMReflection::GetNumericTypeCast(const TCHAR* InTypeStr, const TArray<FF
 #if UENUM_REFL
 			else if (FieldInstance.IsA(FEnumProperty::StaticClass()))
 			{
-				FEnumProperty* EnumFieldInstance = CastField<FEnumProperty>(FieldInstance);
+				FEnumProperty* EnumFieldInstance = FieldInstance.Get<FEnumProperty>();
 				FNumericProperty* UnderlyingProp = EnumFieldInstance->GetUnderlyingProperty();
 				const FFieldClass* UnderlyingPropClass = UnderlyingProp->GetClass();
 
@@ -945,7 +945,7 @@ IMPLEMENT_GENERIC_POINTER_CAST(FText, FTextProperty);
 FVMReflection::operator bool()
 {
 	bool bReturnVal = false;
-	FBoolProperty* BoolProp = CastField<FBoolProperty>(FieldInstance);
+	FBoolProperty* BoolProp = FieldInstance.Get<FBoolProperty>();
 
 	AddCastHistory(TEXT("(bool)"));
 
@@ -994,7 +994,7 @@ FVMReflection::operator FString()
 {
 	FString ReturnVal = TEXT("");
 
-	if (const FEnumProperty* EnumProp = CastField<FEnumProperty>(FieldInstance))
+	if (const FEnumProperty* EnumProp = FieldInstance.Get<FEnumProperty>())
 	{
 		UEnum* TargetEnum = EnumProp->GetEnum();
 		uint8 EnumValue = (uint8)(*this);
@@ -1003,7 +1003,7 @@ FVMReflection::operator FString()
 
 		ReturnVal = TargetEnum->GetNameStringByValue(EnumValue);
 	}
-	else if (const FByteProperty* ByteProp = CastField<FByteProperty>(FieldInstance))
+	else if (const FByteProperty* ByteProp = FieldInstance.Get<FByteProperty>())
 	{
 		UEnum* TargetEnum = ByteProp->Enum;
 		uint8 EnumValue = (uint8)(*this);
@@ -1153,7 +1153,7 @@ FVMReflection::operator TSharedPtr<FScriptArrayHelper>()
 
 	if (ScriptArray != NULL)
 	{
-		ReturnVal = MakeShareable(new FScriptArrayHelper(CastField<FArrayProperty>(FieldInstance), ScriptArray));
+		ReturnVal = MakeShareable(new FScriptArrayHelper(FieldInstance.Get<FArrayProperty>(), ScriptArray));
 	}
 	else
 	{
@@ -1192,7 +1192,7 @@ FVMReflection::operator void*()
 
 FVMReflection& FVMReflection::operator = (bool Value)
 {
-	FBoolProperty* BoolProp = CastField<FBoolProperty>(FieldInstance);
+	FBoolProperty* BoolProp = FieldInstance.Get<FBoolProperty>();
 
 	AddHistory(Value ? TEXT(" = true") : TEXT(" = false"));
 
@@ -1232,7 +1232,7 @@ FVMReflection& FVMReflection::operator = (UObject* Value)
 
 FVMReflection& FVMReflection::operator = (TCHAR* Value)
 {
-	if (FEnumProperty* EnumProp = CastField<FEnumProperty>(FieldInstance))
+	if (FEnumProperty* EnumProp = FieldInstance.Get<FEnumProperty>())
 	{
 		UEnum* TargetEnum = EnumProp->GetEnum();
 		int64 EnumIdx = (int64)TargetEnum->GetIndexByName(Value);
@@ -1246,7 +1246,7 @@ FVMReflection& FVMReflection::operator = (TCHAR* Value)
 			SetError(FString::Printf(TEXT("Name '%s' is not a valid name within enum '%s'."), Value, *TargetEnum->GetName()));
 		}
 	}
-	else if (FByteProperty* ByteProp = CastField<FByteProperty>(FieldInstance))
+	else if (FByteProperty* ByteProp = FieldInstance.Get<FByteProperty>())
 	{
 		UEnum* TargetEnum = ByteProp->Enum;
 		int8 EnumIdx = (int8)TargetEnum->GetIndexByName(Value);
@@ -1260,7 +1260,7 @@ FVMReflection& FVMReflection::operator = (TCHAR* Value)
 			SetError(FString::Printf(TEXT("Name '%s' is not a valid name within enum '%s'."), Value, *TargetEnum->GetName()));
 		}
 	}
-	else if (FSoftObjectProperty* SoftObjProp = CastField<FSoftObjectProperty>(FieldInstance))
+	else if (FSoftObjectProperty* SoftObjProp = FieldInstance.Get<FSoftObjectProperty>())
 	{
 		FSoftObjectPtr* SoftObjPtr = (FSoftObjectPtr*)(*this);
 
@@ -1269,7 +1269,7 @@ FVMReflection& FVMReflection::operator = (TCHAR* Value)
 			*SoftObjPtr = FSoftObjectPath(Value);
 		}
 	}
-	else if (FNameProperty* NameProp = CastField<FNameProperty>(FieldInstance))
+	else if (FNameProperty* NameProp = FieldInstance.Get<FNameProperty>())
 	{
 		FName* TargetName = (FName*)(*this);
 
@@ -1330,7 +1330,7 @@ TValueOrError<FString, FString> FVMReflection::GetValueAsString()
 	{
 		// @todo #JohnB: Static arrays
 
-		FProperty* Prop = CastField<FProperty>(FieldInstance);
+		FProperty* Prop = FieldInstance.Get<FProperty>();
 		UObject* Obj = (UObject*)BaseAddress;
 
 		if (Prop != nullptr && Obj != nullptr)
diff --git a/Engine/Plugins/NetcodeUnitTest/NetcodeUnitTest/Source/NetcodeUnitTest/Public/NUTUtilReflection.h b/Engine/Plugins/NetcodeUnitTest/NetcodeUnitTest/Source/NetcodeUnitTest/Public/NUTUtilReflection.h
--- a/Engine/Plugins/NetcodeUnitTest/NetcodeUnitTest/Source/NetcodeUnitTest/Public/NUTUtilReflection.h
+++ b/Engine/Plugins/NetcodeUnitTest/NetcodeUnitTest/Source/NetcodeUnitTest/Public/NUTUtilReflection.h
@@ -619,7 +619,7 @@ private:
 	 */
 	FORCEINLINE bool IsPropertyArray() const
 	{
-		const FProperty* CurProp = CastField<FProperty>(FieldInstance);
+		const FProperty* CurProp = FieldInstance.Get<FProperty>();
 
 		return CurProp != nullptr && (CurProp->ArrayDim > 1 || CastField<FArrayProperty>(CurProp) != nullptr);
 	}
diff --git a/Engine/Source/Developer/BlueprintCompilerCppBackend/Private/BlueprintCompilerCppBackendGatherDependencies.cpp b/Engine/Source/Developer/BlueprintCompilerCppBackend/Private/BlueprintCompilerCppBackendGatherDependencies.cpp
--- a/Engine/Source/Developer/BlueprintCompilerCppBackend/Private/BlueprintCompilerCppBackendGatherDependencies.cpp
+++ b/Engine/Source/Developer/BlueprintCompilerCppBackend/Private/BlueprintCompilerCppBackendGatherDependencies.cpp
@@ -522,13 +522,80 @@ bool FGatherConvertedClassDependencies::WillClassBeConverted(const UBlueprintGen
 	return false;
 }
 
+
+class FFieldCollectorArchive : public FArchive
+{
+	TSet<FField*> VisitedFields;
+	TArray<FFieldVariant>& Fields;
+
+public:
+
+	FFieldCollectorArchive(TArray<FFieldVariant>& OutFields)
+		: Fields(OutFields)
+	{
+		this->SetIsSaving(true);
+		this->SetIsPersistent(false);
+	}
+
+	virtual FArchive& operator << (FField*& InField) override
+	{
+		if (InField && !VisitedFields.Contains(InField))
+		{
+			VisitedFields.Add(InField);
+			Fields.Add(InField);
+			InField->Serialize(*this);
+		}
+		return *this;
+	}
+};
+
+class FSimpleArrayReferenceCollector : public FReferenceCollector
+{
+	TArray<UObject*>& Objects;
+public:
+	FSimpleArrayReferenceCollector(TArray<UObject*>& OutObjects)
+		: Objects(OutObjects)
+	{
+	}
+	virtual bool IsIgnoringArchetypeRef() const override { return false; }
+	virtual bool IsIgnoringTransient() const override { return false; }
+	virtual void HandleObjectReference(UObject*& InObject, const UObject* InReferencingObject, const FProperty* InReferencingProperty)
+	{
+		if (InObject)
+		{
+			Objects.Add(InObject);
+		}
+	}
+};
+
+void GetStructProperties(UStruct* InStruct, TArray<FFieldVariant>& OutFields)
+{
+	FFieldCollectorArchive Ar(OutFields);
+
+	for (FField* Prop = InStruct->ChildProperties; Prop; Prop = Prop->Next)
+	{
+		Ar << Prop;
+	}
+}
 void FGatherConvertedClassDependencies::DependenciesForHeader()
 {
-	TArray<UObject*> ObjectsToCheck;
-	GetObjectsWithOuter(OriginalStruct, ObjectsToCheck, true);
+	TArray<FFieldVariant> ObjectsToCheck;
+	{
+		TArray<UObject*> StructSubobjects;
+		GetObjectsWithOuter(OriginalStruct, StructSubobjects, true);
+		GetStructProperties(OriginalStruct, ObjectsToCheck);
+		for (UObject* SubObj : StructSubobjects)
+		{
+			ObjectsToCheck.Add(SubObj);
+			if (UStruct* SubStruct = Cast<UStruct>(SubObj))
+			{
+				GetStructProperties(SubStruct, ObjectsToCheck);
+			}
+		}
+	}
 
 	TArray<UObject*> NeededObjects;
-	FReferenceFinder HeaderReferenceFinder(NeededObjects, nullptr, false, false, true, false);
+	FSimpleArrayReferenceCollector HeaderReferenceFinder(NeededObjects);
 
 	auto ShouldIncludeHeaderFor = [this](UObject* InObj)->bool
 	{
@@ -547,9 +614,9 @@ void FGatherConvertedClassDependencies::DependenciesForHeader()
 		return false;
 	};
 
-	for (UObject* Obj : ObjectsToCheck)
+	for (FFieldVariant& Obj : ObjectsToCheck)
 	{
-		const FProperty* Property = CastField<const FProperty>(Obj);
+		const FProperty* Property = Obj.Get<FProperty>();
 		if (const FArrayProperty* ArrayProperty = CastField<FArrayProperty>(Property))
 		{
 			Property = ArrayProperty->Inner;
@@ -601,7 +668,9 @@ void FGatherConvertedClassDependencies::DependenciesForHeader()
 			}
 			else
 			{
-				HeaderReferenceFinder.FindReferences(Obj);
+				FField* Field = Obj.Get<FField>();	
+				check(Field);
+				Field->AddReferencedObjects(HeaderReferenceFinder);
 			}
 		}
 	}
diff --git a/Engine/Source/Developer/BlueprintCompilerCppBackend/Private/BlueprintCompilerCppBackendUtils.cpp b/Engine/Source/Developer/BlueprintCompilerCppBackend/Private/BlueprintCompilerCppBackendUtils.cpp
--- a/Engine/Source/Developer/BlueprintCompilerCppBackend/Private/BlueprintCompilerCppBackendUtils.cpp
+++ b/Engine/Source/Developer/BlueprintCompilerCppBackend/Private/BlueprintCompilerCppBackendUtils.cpp
@@ -1839,7 +1839,7 @@ FString FEmitHelper::GenerateGetPropertyByName(FEmitterLocalContext& EmitterCont
 	if (bUseStaticVariables)
 	{
 		const FString PropertyWeakPtrName = EmitterContext.GenerateUniqueLocalName();
-		EmitterContext.AddLine(FString::Printf(TEXT("static TWeakObjectPtr<FProperty> %s{};"), *PropertyWeakPtrName));
+		EmitterContext.AddLine(FString::Printf(TEXT("static TWeakFieldPtr<FProperty> %s{};"), *PropertyWeakPtrName));
 
 		EmitterContext.AddLine(FString::Printf(TEXT("const FProperty* %s = %s.Get();"), *PropertyPtrName, *PropertyWeakPtrName));
 		EmitterContext.AddLine(FString::Printf(TEXT("if (nullptr == %s)"), *PropertyPtrName));
diff --git a/Engine/Source/Editor/KismetCompiler/Private/KismetCompiler.cpp b/Engine/Source/Editor/KismetCompiler/Private/KismetCompiler.cpp
--- a/Engine/Source/Editor/KismetCompiler/Private/KismetCompiler.cpp
+++ b/Engine/Source/Editor/KismetCompiler/Private/KismetCompiler.cpp
@@ -4497,19 +4497,35 @@ void FKismetCompilerContext::CompileFunctions(EInternalCompilerFlags InternalFla
 		class FSignatureArchiveCrc32 : public FArchiveObjectCrc32
 		{
 		public:
-			static bool IsInnerProperty(const UObject* Object)
+
+			static bool IsInnerProperty(const FField* Field)
 			{
-				const FProperty* Property = CastField<const FProperty>(Object);
+				const FProperty* Property = CastField<const FProperty>(Field);
 				return Property // check arrays
 					&& Cast<const UFunction>(Property->GetOwnerStruct())
 					&& !Property->HasAnyPropertyFlags(CPF_Parm);
 			}
 
+			virtual FArchive& operator<<(FField*& Field) override
+			{
+				FArchive& Ar = *this;
+				if (Field && !IsInnerProperty(Field))
+				{					
+					FString UniqueName = GetPathNameSafe(Field);
+					Ar << UniqueName;
+					if (Field->IsIn(RootObject))
+					{
+						Field->Serialize(Ar);
+					}
+				}
+				return Ar;
+			}
+
 			virtual FArchive& operator<<(UObject*& Object) override
 			{
 				FArchive& Ar = *this;
 
-				if (Object && !IsInnerProperty(Object))
+				if (Object)
 				{
 					// Names of functions and properties are significant.
 					FString UniqueName = GetPathNameSafe(Object);
@@ -4547,6 +4563,13 @@ void FKismetCompilerContext::CompileFunctions(EInternalCompilerFlags InternalFla
 						ChildrenIter = ChildrenIter->Next;
 					}
 
+					FField* ChildPropIter = Struct->ChildProperties;
+					while (ChildPropIter)
+					{
+						Ar << ChildPropIter;
+						ChildPropIter = ChildPropIter->Next;
+					}
+
 					if (UFunction* Function = Cast<UFunction>(Struct))
 					{
 						Ar << Function->FunctionFlags;
diff --git a/Engine/Source/Programs/UnrealHeaderTool/Private/CodeGenerator.cpp b/Engine/Source/Programs/UnrealHeaderTool/Private/CodeGenerator.cpp
--- a/Engine/Source/Programs/UnrealHeaderTool/Private/CodeGenerator.cpp
+++ b/Engine/Source/Programs/UnrealHeaderTool/Private/CodeGenerator.cpp
@@ -5000,7 +5000,7 @@ void FNativeClassHeaderGenerator::ExportNativeFunctions(FOutputDevice& OutGenera
 		UScriptStruct* SparseClassDataStruct = FindObjectSafe<UScriptStruct>(ANY_PACKAGE, *SparseClassDataString);
 		while (SparseClassDataStruct != nullptr)
 		{
-			const FProperty* Child = CastField<FProperty>(SparseClassDataStruct->Children);
+			const FProperty* Child = CastField<FProperty>(SparseClassDataStruct->ChildProperties);
 			while (Child)
 			{
 				FString ReturnExtendedType;
diff --git a/Engine/Source/Runtime/CoreUObject/Private/UObject/Field.cpp b/Engine/Source/Runtime/CoreUObject/Private/UObject/Field.cpp
--- a/Engine/Source/Runtime/CoreUObject/Private/UObject/Field.cpp
+++ b/Engine/Source/Runtime/CoreUObject/Private/UObject/Field.cpp
@@ -516,7 +516,7 @@ bool FField::IsValidLowLevel() const
 	return this != nullptr;
 }
 
-bool FField::IsIn(UObject* InOwner) const
+bool FField::IsIn(const UObject* InOwner) const
 {
 	check(InOwner);
 	UObject* OwnerObject = GetOwnerUObject();
@@ -534,7 +534,7 @@ bool FField::IsIn(UObject* InOwner) const
 	return false;
 }
 
-bool FField::IsIn(FField* InOwner) const
+bool FField::IsIn(const FField* InOwner) const
 {
 	for (FField* OwnerField = GetOwner<FField>(); OwnerField; OwnerField = OwnerField->GetOwner<FField>())
 	{
diff --git a/Engine/Source/Runtime/CoreUObject/Public/UObject/Field.h b/Engine/Source/Runtime/CoreUObject/Public/UObject/Field.h
--- a/Engine/Source/Runtime/CoreUObject/Public/UObject/Field.h
+++ b/Engine/Source/Runtime/CoreUObject/Public/UObject/Field.h
@@ -384,8 +384,8 @@ public:
 	bool IsRooted() const;
 	bool IsNative() const;
 	bool IsValidLowLevel() const;
-	bool IsIn(UObject* InOwner) const;
-	bool IsIn(FField* InOwner) const;
+	bool IsIn(const UObject* InOwner) const;
+	bool IsIn(const FField* InOwner) const;
 	FLinkerLoad* GetLinker() const;	
 	// End UObject interface
 
@@ -813,18 +813,6 @@ FUNCTION_NON_NULL_RETURN_END
 #endif // !DO_CHECK
 }
 
-template<typename FieldType>
-FORCEINLINE FieldType* CastField(FFieldVariant& Src)
-{
-	return CastField<FieldType>(Src.ToField());
-}
-
-template<typename FieldType>
-FORCEINLINE const FieldType* CastField(const FFieldVariant& Src)
-{
-	return CastField<FieldType>(Src.ToField());
-}
-
 /**
  * Helper function for serializing FField to an archive. This function fully serializes the field and its properties.
  */
diff --git a/Engine/Source/Runtime/CoreUObject/Public/UObject/WeakFieldPtr.h b/Engine/Source/Runtime/CoreUObject/Public/UObject/WeakFieldPtr.h
--- a/Engine/Source/Runtime/CoreUObject/Public/UObject/WeakFieldPtr.h
+++ b/Engine/Source/Runtime/CoreUObject/Public/UObject/WeakFieldPtr.h
@@ -75,10 +75,10 @@ public:
 	* @param Object object to create a weak pointer to
 	**/
 	template<class U>
-	FORCEINLINE typename TEnableIf<!TLosesQualifiersFromTo<U, T>::Value>::Type operator=(U* InField)
+	FORCEINLINE typename TEnableIf<!TLosesQualifiersFromTo<U, T>::Value>::Type operator=(const U* InField)
 	{
 		Owner = InField ? InField->GetOwnerUObject() : (UObject*)nullptr;
-		Field = InField;
+		Field = (U*)InField;
 	}
 
 	/**
diff --git a/Engine/Source/Runtime/Engine/Public/GeneratedCodeHelpers.h b/Engine/Source/Runtime/Engine/Public/GeneratedCodeHelpers.h
--- a/Engine/Source/Runtime/Engine/Public/GeneratedCodeHelpers.h
+++ b/Engine/Source/Runtime/Engine/Public/GeneratedCodeHelpers.h
@@ -7,6 +7,7 @@
 
 // Common includes
 #include "UObject/Stack.h"
+#include "UObject/WeakFieldPtr.h"
 #include "Blueprint/BlueprintSupport.h"
 #include "Engine/BlueprintGeneratedClass.h"
 #include "Engine/UserDefinedStruct.h"
diff --git a/Engine/Source/Runtime/Engine/Public/Net/UnrealNetwork.h b/Engine/Source/Runtime/Engine/Public/Net/UnrealNetwork.h
--- a/Engine/Source/Runtime/Engine/Public/Net/UnrealNetwork.h
+++ b/Engine/Source/Runtime/Engine/Public/Net/UnrealNetwork.h
@@ -186,7 +186,7 @@ static FProperty* GetReplicatedProperty(const UClass* CallingClass, const UClass
 /** This macro is used by nativized code (DynamicClasses), so the Property may be recreated. */
 #define DOREPLIFETIME_DIFFNAMES(c,v, n) \
 { \
-	static TWeakObjectPtr<FProperty> __swp##v{};							\
+	static TWeakFieldPtr<FProperty> __swp##v{};							\
 	const FProperty* sp##v = __swp##v.Get();								\
 	if (nullptr == sp##v)													\
 	{																		\
