Fix compile errors in NetcodeUnitTest.

diff --git a/Engine/Plugins/NetcodeUnitTest/NetcodeUnitTest/Source/NetcodeUnitTest/Classes/ClientUnitTest.h b/Engine/Plugins/NetcodeUnitTest/NetcodeUnitTest/Source/NetcodeUnitTest/Classes/ClientUnitTest.h
--- a/Engine/Plugins/NetcodeUnitTest/NetcodeUnitTest/Source/NetcodeUnitTest/Classes/ClientUnitTest.h
+++ b/Engine/Plugins/NetcodeUnitTest/NetcodeUnitTest/Source/NetcodeUnitTest/Classes/ClientUnitTest.h
@@ -116,7 +116,7 @@ protected:
 
 	/** The object which handles implementation of the minimal client */
 	UPROPERTY()
-	UMinimalClient* MinClient;
+	TObjectPtr<UMinimalClient> MinClient;
 
 	/** Whether or not the initial connect of the minimal client was triggered */
 	bool bTriggerredInitialConnect;
diff --git a/Engine/Plugins/NetcodeUnitTest/NetcodeUnitTest/Source/NetcodeUnitTest/Classes/NUTGlobals.h b/Engine/Plugins/NetcodeUnitTest/NetcodeUnitTest/Source/NetcodeUnitTest/Classes/NUTGlobals.h
--- a/Engine/Plugins/NetcodeUnitTest/NetcodeUnitTest/Source/NetcodeUnitTest/Classes/NUTGlobals.h
+++ b/Engine/Plugins/NetcodeUnitTest/NetcodeUnitTest/Source/NetcodeUnitTest/Classes/NUTGlobals.h
@@ -45,7 +45,7 @@ public:
 public:
 	/** For NUTActor - the UNetConnection that 'watch' events will be sent to */
 	UPROPERTY()
-	UNetConnection* EventWatcher;
+	TObjectPtr<UNetConnection> EventWatcher;
 
 	/** For ClientUnitTest - the (+10) incremented port number, for each server instance */
 	UPROPERTY()
diff --git a/Engine/Plugins/NetcodeUnitTest/NetcodeUnitTest/Source/NetcodeUnitTest/Classes/UnitTest.h b/Engine/Plugins/NetcodeUnitTest/NetcodeUnitTest/Source/NetcodeUnitTest/Classes/UnitTest.h
--- a/Engine/Plugins/NetcodeUnitTest/NetcodeUnitTest/Source/NetcodeUnitTest/Classes/UnitTest.h
+++ b/Engine/Plugins/NetcodeUnitTest/NetcodeUnitTest/Source/NetcodeUnitTest/Classes/UnitTest.h
@@ -209,7 +209,7 @@ protected:
 
 	/** UnitTask's which must be run before different stages of the unit test can execute */
 	UPROPERTY()
-	TArray<UUnitTask*> UnitTasks;
+	TArray<TObjectPtr<UUnitTask>> UnitTasks;
 
 	/** Marks the state of met unit task requirement flags and active unit task blocking flags */
 	EUnitTaskFlags UnitTaskState;
diff --git a/Engine/Plugins/NetcodeUnitTest/NetcodeUnitTest/Source/NetcodeUnitTest/Classes/UnitTestManager.h b/Engine/Plugins/NetcodeUnitTest/NetcodeUnitTest/Source/NetcodeUnitTest/Classes/UnitTestManager.h
--- a/Engine/Plugins/NetcodeUnitTest/NetcodeUnitTest/Source/NetcodeUnitTest/Classes/UnitTestManager.h
+++ b/Engine/Plugins/NetcodeUnitTest/NetcodeUnitTest/Source/NetcodeUnitTest/Classes/UnitTestManager.h
@@ -64,15 +64,15 @@ public:
 
 	/** Holds a list of unit tests pending execution */
 	UPROPERTY()
-	TArray<UClass*> PendingUnitTests;
+	TArray<TObjectPtr<UClass>> PendingUnitTests;
 
 	/** Holds a list of currently active unit tests */
 	UPROPERTY()
-	TArray<UUnitTest*> ActiveUnitTests;
+	TArray<TObjectPtr<UUnitTest>> ActiveUnitTests;
 
 	/** Unit tests which are finished, and are kept around until printing the final summary */
 	UPROPERTY()
-	TArray<UUnitTest*> FinishedUnitTests;
+	TArray<TObjectPtr<UUnitTest>> FinishedUnitTests;
 
 
 	/** If a unit test was aborted on its first run, strictly cap all first-run unit tests to one at a time */
diff --git a/Engine/Plugins/NetcodeUnitTest/NetcodeUnitTest/Source/NetcodeUnitTest/Classes/UnitTests/VMReflection.h b/Engine/Plugins/NetcodeUnitTest/NetcodeUnitTest/Source/NetcodeUnitTest/Classes/UnitTests/VMReflection.h
--- a/Engine/Plugins/NetcodeUnitTest/NetcodeUnitTest/Source/NetcodeUnitTest/Classes/UnitTests/VMReflection.h
+++ b/Engine/Plugins/NetcodeUnitTest/NetcodeUnitTest/Source/NetcodeUnitTest/Classes/UnitTests/VMReflection.h
@@ -34,7 +34,7 @@ class UVMTestClassA : public UObject
 
 public:
 	UPROPERTY()
-	UObject*	AObjectRef;
+	TObjectPtr<UObject> AObjectRef;
 
 	UPROPERTY()
 	uint8		ByteProp;
@@ -95,7 +95,7 @@ public:
 	uint8		BytePropArray[4];
 
 	UPROPERTY()
-	UObject*	ObjectPropArray[2];
+	TObjectPtr<UObject> ObjectPropArray[2];
 
 	UPROPERTY()
 	TArray<uint8>		DynBytePropArray;
@@ -104,7 +104,7 @@ public:
 	TArray<bool>		DynBoolPropArray;
 
 	UPROPERTY()
-	TArray<UObject*>	DynObjectPropArray;
+	TArray<TObjectPtr<UObject>> DynObjectPropArray;
 
 	UPROPERTY()
 	TArray<FName>		DynNamePropArray;
@@ -143,10 +143,10 @@ public:
 	TArray<FText>		DynTextPropArray;
 
 	UPROPERTY()
-	TArray<UClass*>		DynClassPropArray;
+	TArray<TObjectPtr<UClass>> DynClassPropArray;
 
 	UPROPERTY()
-	TArray<APawn*>		DynPawnPropArray;
+	TArray<TObjectPtr<APawn>> DynPawnPropArray;
 
 
 	UPROPERTY()
@@ -166,6 +166,6 @@ class UVMTestClassB : public UObject
 
 public:
 	UPROPERTY()
-	UObject*	BObjectRef;
+	TObjectPtr<UObject> BObjectRef;
 };
 
diff --git a/Engine/Plugins/NetcodeUnitTest/NetcodeUnitTest/Source/NetcodeUnitTest/Private/NUTActor.cpp b/Engine/Plugins/NetcodeUnitTest/NetcodeUnitTest/Source/NetcodeUnitTest/Private/NUTActor.cpp
--- a/Engine/Plugins/NetcodeUnitTest/NetcodeUnitTest/Source/NetcodeUnitTest/Private/NUTActor.cpp
+++ b/Engine/Plugins/NetcodeUnitTest/NetcodeUnitTest/Source/NetcodeUnitTest/Private/NUTActor.cpp
@@ -302,7 +302,7 @@ void ANUTActor::NotifyPostLoadMap(UWorld* LoadedWorld)
 bool ANUTActor::VerifyEventWatcher()
 {
 	bool bVerified = false;
-	UNetConnection*& EventWatcher = UNUTGlobals::Get().EventWatcher;
+	TObjectPtr<UNetConnection>& EventWatcher = UNUTGlobals::Get().EventWatcher;
 
 	if (EventWatcher != nullptr)
 	{
diff --git a/Engine/Plugins/NetcodeUnitTest/NetcodeUnitTest/Source/NetcodeUnitTest/Private/UnitTestManager.cpp b/Engine/Plugins/NetcodeUnitTest/NetcodeUnitTest/Source/NetcodeUnitTest/Private/UnitTestManager.cpp
--- a/Engine/Plugins/NetcodeUnitTest/NetcodeUnitTest/Source/NetcodeUnitTest/Private/UnitTestManager.cpp
+++ b/Engine/Plugins/NetcodeUnitTest/NetcodeUnitTest/Source/NetcodeUnitTest/Private/UnitTestManager.cpp
@@ -752,7 +752,7 @@ void UUnitTestManager::NotifyLogWindowClosed(const TSharedRef<SWindow>& ClosedWi
 	else
 	{
 		// Match the log window to a unit test
-		UUnitTest** CurUnitTestRef = ActiveUnitTests.FindByPredicate(
+		TObjectPtr<UUnitTest>* CurUnitTestRef = ActiveUnitTests.FindByPredicate(
 			[&](const UUnitTest* Element)
 			{
 				return Element->LogWindow == ClosedWindow;
@@ -1591,7 +1591,7 @@ bool UUnitTestManager::Exec(UWorld* InWorld, const TCHAR* Cmd, FOutputDevice& Ar
 	// If this command was triggered within a specific unit test (as identified by InWorld), abort it
 	else if (UnitTestName == TEXT("abort"))
 	{
-		UUnitTest** UnitTestRef = (InWorld != nullptr ? ActiveUnitTests.FindByPredicate(FindUnitTestByWorld) :
+		TObjectPtr<UUnitTest>* UnitTestRef = (InWorld != nullptr ? ActiveUnitTests.FindByPredicate(FindUnitTestByWorld) :
 									ActiveUnitTests.FindByPredicate(FindUnitTestByLog));
 
 		UUnitTest* AbortUnitTest = (UnitTestRef != nullptr ? *UnitTestRef : nullptr);
@@ -1610,7 +1610,7 @@ bool UUnitTestManager::Exec(UWorld* InWorld, const TCHAR* Cmd, FOutputDevice& Ar
 	// If this command was triggered within a specific unit test, restart the unit test
 	else if (UnitTestName == TEXT("restart"))
 	{
-		UUnitTest** UnitTestRef = (InWorld != nullptr ? ActiveUnitTests.FindByPredicate(FindUnitTestByWorld) :
+		TObjectPtr<UUnitTest>* UnitTestRef = (InWorld != nullptr ? ActiveUnitTests.FindByPredicate(FindUnitTestByWorld) :
 									ActiveUnitTests.FindByPredicate(FindUnitTestByLog));
 
 		UUnitTest* RestartUnitTest = (UnitTestRef != nullptr ? *UnitTestRef : nullptr);
@@ -1663,7 +1663,7 @@ bool UUnitTestManager::Exec(UWorld* InWorld, const TCHAR* Cmd, FOutputDevice& Ar
 	// Debug unit test commands
 	else if (UnitTestName == TEXT("debug"))
 	{
-		UUnitTest** TargetUnitTestRef = (InWorld != nullptr ? ActiveUnitTests.FindByPredicate(
+		TObjectPtr<UUnitTest>* TargetUnitTestRef = (InWorld != nullptr ? ActiveUnitTests.FindByPredicate(
 			[&InWorld](const UUnitTest* InElement)
 			{
 				auto CurUnitTest = Cast<UClientUnitTest>(InElement);
