NetcodeUnitTest updates/fixups.

diff --git a/Engine/Plugins/NetcodeUnitTest/NUTUnrealEngine/Source/NUTUnrealEngine/Classes/UnitTests/FTextCrash.h b/Engine/Plugins/NetcodeUnitTest/NUTUnrealEngine/Source/NUTUnrealEngine/Classes/UnitTests/FTextCrash.h
--- a/Engine/Plugins/NetcodeUnitTest/NUTUnrealEngine/Source/NUTUnrealEngine/Classes/UnitTests/FTextCrash.h
+++ b/Engine/Plugins/NetcodeUnitTest/NUTUnrealEngine/Source/NUTUnrealEngine/Classes/UnitTests/FTextCrash.h
@@ -2,7 +2,7 @@
 
 #pragma once
 
-#include "CoreMinimal.h"
+// Includes
 #include "UObject/ObjectMacros.h"
 #include "ClientUnitTest.h"
 
diff --git a/Engine/Plugins/NetcodeUnitTest/NUTUnrealEngine/Source/NUTUnrealEngine/Classes/UnitTests/NetBitsTest.h b/Engine/Plugins/NetcodeUnitTest/NUTUnrealEngine/Source/NUTUnrealEngine/Classes/UnitTests/NetBitsTest.h
--- a/Engine/Plugins/NetcodeUnitTest/NUTUnrealEngine/Source/NUTUnrealEngine/Classes/UnitTests/NetBitsTest.h
+++ b/Engine/Plugins/NetcodeUnitTest/NUTUnrealEngine/Source/NUTUnrealEngine/Classes/UnitTests/NetBitsTest.h
@@ -2,7 +2,7 @@
 
 #pragma once
 
-#include "CoreMinimal.h"
+// Includes
 #include "UObject/ObjectMacros.h"
 #include "UnitTest.h"
 
diff --git a/Engine/Plugins/NetcodeUnitTest/NUTUnrealEngine/Source/NUTUnrealEngine/Classes/UnitTests/Obsolete/UTT61_DebugReplicateData.h b/Engine/Plugins/NetcodeUnitTest/NUTUnrealEngine/Source/NUTUnrealEngine/Classes/UnitTests/Obsolete/UTT61_DebugReplicateData.h
--- a/Engine/Plugins/NetcodeUnitTest/NUTUnrealEngine/Source/NUTUnrealEngine/Classes/UnitTests/Obsolete/UTT61_DebugReplicateData.h
+++ b/Engine/Plugins/NetcodeUnitTest/NUTUnrealEngine/Source/NUTUnrealEngine/Classes/UnitTests/Obsolete/UTT61_DebugReplicateData.h
@@ -2,15 +2,20 @@
 
 #pragma once
 
-#include "CoreMinimal.h"
+// Includes
 #include "UObject/ObjectMacros.h"
+#include "UObject/WeakObjectPtrTemplates.h"
+#include "UObject/ObjectPtr.h"
 #include "ClientUnitTest.h"
 
 #include "UTT61_DebugReplicateData.generated.h"
 
+
+// Forward declarations
 class AActor;
 class UActorChannel;
 
+
 /**
  * Implements UTT GameplayDebuggingComponent/ServerReplicateData crash
  *
@@ -37,10 +42,12 @@ class UUTT61_DebugReplicateData : public UClientUnitTest
 
 protected:
 	/** Reference to the GameplayDebuggingReplicator class */
-	UClass* RepClass;
+	UPROPERTY()
+	TObjectPtr<UClass> RepClass;
 
 private:
 	/** Once the GameplayDebuggingReplicator is received, cache it */
+	UPROPERTY()
 	TWeakObjectPtr<AActor> Replicator;
 
 public:
diff --git a/Engine/Plugins/NetcodeUnitTest/NUTUnrealEngine/Source/NUTUnrealEngine/Private/NUTUnrealEngine.cpp b/Engine/Plugins/NetcodeUnitTest/NUTUnrealEngine/Source/NUTUnrealEngine/Private/NUTUnrealEngine.cpp
--- a/Engine/Plugins/NetcodeUnitTest/NUTUnrealEngine/Source/NUTUnrealEngine/Private/NUTUnrealEngine.cpp
+++ b/Engine/Plugins/NetcodeUnitTest/NUTUnrealEngine/Source/NUTUnrealEngine/Private/NUTUnrealEngine.cpp
@@ -1,8 +1,7 @@
 // Copyright Epic Games, Inc. All Rights Reserved.
 
-#include "CoreMinimal.h"
+// Includes
 #include "Modules/ModuleManager.h"
-
 #include "INUTUnrealEngine.h"
 
 #include "UnrealEngineEnvironment.h"
diff --git a/Engine/Plugins/NetcodeUnitTest/NUTUnrealEngine/Source/NUTUnrealEngine/Public/INUTUnrealEngine.h b/Engine/Plugins/NetcodeUnitTest/NUTUnrealEngine/Source/NUTUnrealEngine/Public/INUTUnrealEngine.h
--- a/Engine/Plugins/NetcodeUnitTest/NUTUnrealEngine/Source/NUTUnrealEngine/Public/INUTUnrealEngine.h
+++ b/Engine/Plugins/NetcodeUnitTest/NUTUnrealEngine/Source/NUTUnrealEngine/Public/INUTUnrealEngine.h
@@ -2,7 +2,7 @@
 
 #pragma once
 
-#include "CoreMinimal.h"
+// Includes
 #include "Modules/ModuleInterface.h"
 #include "Modules/ModuleManager.h"
 
diff --git a/Engine/Plugins/NetcodeUnitTest/NUTUnrealEngine/Source/NUTUnrealEngine/Public/UnrealEngineEnvironment.h b/Engine/Plugins/NetcodeUnitTest/NUTUnrealEngine/Source/NUTUnrealEngine/Public/UnrealEngineEnvironment.h
--- a/Engine/Plugins/NetcodeUnitTest/NUTUnrealEngine/Source/NUTUnrealEngine/Public/UnrealEngineEnvironment.h
+++ b/Engine/Plugins/NetcodeUnitTest/NUTUnrealEngine/Source/NUTUnrealEngine/Public/UnrealEngineEnvironment.h
@@ -2,12 +2,15 @@
 
 #pragma once
 
-#include "CoreMinimal.h"
+// Includes
 #include "Misc/App.h"
 #include "UnitTestEnvironment.h"
 
+
+// Forward declarations
 enum class EUnitTestFlags : uint32;
 
+
 /**
  * Implements unit test environment settings, for ShooterGame
  */
diff --git a/Engine/Plugins/NetcodeUnitTest/NetcodeUnitTest/Source/NetcodeUnitTest/Classes/ClientUnitTest.h b/Engine/Plugins/NetcodeUnitTest/NetcodeUnitTest/Source/NetcodeUnitTest/Classes/ClientUnitTest.h
--- a/Engine/Plugins/NetcodeUnitTest/NetcodeUnitTest/Source/NetcodeUnitTest/Classes/ClientUnitTest.h
+++ b/Engine/Plugins/NetcodeUnitTest/NetcodeUnitTest/Source/NetcodeUnitTest/Classes/ClientUnitTest.h
@@ -2,9 +2,11 @@
 
 #pragma once
 
+// Includes
 #include "Templates/SubclassOf.h"
 #include "UObject/ObjectMacros.h"
 #include "Engine/EngineBaseTypes.h"
+#include "UObject/ObjectPtr.h"
 
 #include "NetcodeUnitTest.h"
 #include "ProcessUnitTest.h"
@@ -82,7 +84,8 @@ protected:
 	FString BaseClientParameters;
 
 	/** Actors the server is allowed replicate to client (requires AllowActors flag). Use NotifyAllowNetActor for conditional allows. */
-	TArray<UClass*> AllowedClientActors;
+	UPROPERTY()
+	TArray<TObjectPtr<UClass>> AllowedClientActors;
 
 	/** Clientside RPC's that should be allowed to execute (requires minimal client NotifyProcessNetEvent flag) */
 	TArray<FString> AllowedClientRPCs;
@@ -125,6 +128,7 @@ protected:
 	/** Client state variables */
 protected:
 	/** Stores a reference to the replicated PlayerController (if set to wait for this), after NotifyHandleClientPlayer */
+	UPROPERTY()
 	TWeakObjectPtr<APlayerController> UnitPC;
 
 	/** Whether or not the UnitPC Pawn was fully setup (requires EUnitTestFlags::RequirePawn) */
@@ -134,12 +138,14 @@ protected:
 	bool bUnitPlayerStateSetup;
 
 	/** If EUnitTestFlags::RequireNUTActor is set, stores a reference to the replicated NUTActor */
+	UPROPERTY()
 	TWeakObjectPtr<ANUTActor> UnitNUTActor;
 
 	/** Whether or not UnitNUTActor is fully setup, i.e. has replicated its Owner */
 	bool bUnitNUTActorSetup;
 
 	/** If EUnitTestFlags::RequireBeacon is set, stores a reference to the replicated beacon */
+	UPROPERTY()
 	TWeakObjectPtr<AActor> UnitBeacon;
 
 	/** If EUnitTestFlags::RequirePing is true, whether or not we have already received the pong */
@@ -149,10 +155,23 @@ protected:
 	bool bPendingNetworkFailure;
 
 
-	/** Static variables */
 private:
-	/** Static reference to the OnlineBeaconClient static class */
-	static UClass* OnlineBeaconClass;
+	/** Cached reference to the OnlineBeaconClient static class */
+	UPROPERTY()
+	TObjectPtr<UClass> OnlineBeaconClass_Private;
+
+
+	void InitOnlineBeaconClass();
+
+	UClass* GetOnlineBeaconClass()
+	{
+		if (!OnlineBeaconClass_Private)
+		{
+			InitOnlineBeaconClass();
+		}
+
+		return OnlineBeaconClass_Private;
+	}
 
 
 	/**
diff --git a/Engine/Plugins/NetcodeUnitTest/NetcodeUnitTest/Source/NetcodeUnitTest/Classes/MinimalClient.h b/Engine/Plugins/NetcodeUnitTest/NetcodeUnitTest/Source/NetcodeUnitTest/Classes/MinimalClient.h
--- a/Engine/Plugins/NetcodeUnitTest/NetcodeUnitTest/Source/NetcodeUnitTest/Classes/MinimalClient.h
+++ b/Engine/Plugins/NetcodeUnitTest/NetcodeUnitTest/Source/NetcodeUnitTest/Classes/MinimalClient.h
@@ -54,9 +54,8 @@
  *			perhaps keeping the MCP side within the UnitTask
  */
 
-
-#include "CoreMinimal.h"
-
+// Includes
+#include "UObject/ObjectPtr.h"
 #include "Engine/NetConnection.h"
 #include "Engine/PendingNetGame.h"
 
@@ -69,6 +68,11 @@
 
 // Forward declarations
 class FFuncReflection;
+class UMinimalClient;
+class UUnitTestNetConnection;
+class UUnitTestChannel;
+class UUnitTestActorChannel;
+class UUnitTestPackageMap;
 
 
 // Delegates
@@ -138,7 +142,7 @@ DECLARE_DELEGATE(FOnRPCFailure);
  */
 struct FMinClientParms
 {
-	friend class UMinimalClient;
+	friend UMinimalClient;
 
 
 	/** The flags used for configuring the minimal client */
@@ -210,7 +214,7 @@ protected:
  */
 struct FMinClientHooks
 {
-	friend class UMinimalClient;
+	friend UMinimalClient;
 
 
 	/** Delegate for notifying of successful minimal client connection*/
@@ -307,10 +311,10 @@ UCLASS()
 class NETCODEUNITTEST_API UMinimalClient :
 	public UObject, public FNetworkNotify, public FProtMinClientParms, public FProtMinClientHooks, public FUnitLogRedirect
 {
-	friend class UUnitTestNetConnection;
-	friend class UUnitTestChannel;
-	friend class UUnitTestActorChannel;
-	friend class UUnitTestPackageMap;
+	friend UUnitTestNetConnection;
+	friend UUnitTestChannel;
+	friend UUnitTestActorChannel;
+	friend UUnitTestPackageMap;
 
 	GENERATED_UCLASS_BODY()
 
@@ -553,13 +557,16 @@ protected:
 	bool bConnected;
 
 	/** Stores a reference to the created empty world, for execution and later cleanup */
-	UWorld* UnitWorld;
+	UPROPERTY()
+	TObjectPtr<UWorld> UnitWorld;
 
 	/** Stores a reference to the created unit test net driver, for execution and later cleanup */
-	UNetDriver* UnitNetDriver;
+	UPROPERTY()
+	TObjectPtr<UNetDriver> UnitNetDriver;
 
 	/** Stores a reference to the server connection */
-	UNetConnection* UnitConn;
+	UPROPERTY()
+	TObjectPtr<UNetConnection> UnitConn;
 
 	/** If notifying of net actor creation, this keeps track of new actor channel indexes pending notification */
 	TArray<int32> PendingNetActorChans;
diff --git a/Engine/Plugins/NetcodeUnitTest/NetcodeUnitTest/Source/NetcodeUnitTest/Classes/NUTActor.h b/Engine/Plugins/NetcodeUnitTest/NetcodeUnitTest/Source/NetcodeUnitTest/Classes/NUTActor.h
--- a/Engine/Plugins/NetcodeUnitTest/NetcodeUnitTest/Source/NetcodeUnitTest/Classes/NUTActor.h
+++ b/Engine/Plugins/NetcodeUnitTest/NetcodeUnitTest/Source/NetcodeUnitTest/Classes/NUTActor.h
@@ -2,7 +2,7 @@
 
 #pragma once
 
-#include "CoreMinimal.h"
+// Includes
 #include "UObject/ObjectMacros.h"
 #include "Misc/CoreMisc.h"
 #include "GameFramework/Actor.h"
@@ -12,6 +12,7 @@
 
 #include "NUTActor.generated.h"
 
+// Forward declarations
 class ANUTActor;
 
 /**
diff --git a/Engine/Plugins/NetcodeUnitTest/NetcodeUnitTest/Source/NetcodeUnitTest/Classes/NUTGlobals.h b/Engine/Plugins/NetcodeUnitTest/NetcodeUnitTest/Source/NetcodeUnitTest/Classes/NUTGlobals.h
--- a/Engine/Plugins/NetcodeUnitTest/NetcodeUnitTest/Source/NetcodeUnitTest/Classes/NUTGlobals.h
+++ b/Engine/Plugins/NetcodeUnitTest/NetcodeUnitTest/Source/NetcodeUnitTest/Classes/NUTGlobals.h
@@ -2,9 +2,10 @@
 
 #pragma once
 
-#include "CoreMinimal.h"
+// Includes
 #include "UObject/ObjectMacros.h"
 #include "UObject/Package.h"
+#include "UObject/ObjectPtr.h"
 
 #include "NUTGlobals.generated.h"
 
diff --git a/Engine/Plugins/NetcodeUnitTest/NetcodeUnitTest/Source/NetcodeUnitTest/Classes/Net/UnitTestActorChannel.h b/Engine/Plugins/NetcodeUnitTest/NetcodeUnitTest/Source/NetcodeUnitTest/Classes/Net/UnitTestActorChannel.h
--- a/Engine/Plugins/NetcodeUnitTest/NetcodeUnitTest/Source/NetcodeUnitTest/Classes/Net/UnitTestActorChannel.h
+++ b/Engine/Plugins/NetcodeUnitTest/NetcodeUnitTest/Source/NetcodeUnitTest/Classes/Net/UnitTestActorChannel.h
@@ -2,9 +2,10 @@
 
 #pragma once
 
-#include "CoreMinimal.h"
+// Includes
 #include "UObject/ObjectMacros.h"
 #include "Engine/ActorChannel.h"
+#include "UObject/ObjectPtr.h"
 
 #include "UnitTestActorChannel.generated.h"
 
@@ -34,6 +35,7 @@ public:
 
 private:
 	/** Cached referenced to the minimal client that owns this actor channel */
-	UMinimalClient* MinClient;
+	UPROPERTY()
+	TObjectPtr<UMinimalClient> MinClient;
 };
 
diff --git a/Engine/Plugins/NetcodeUnitTest/NetcodeUnitTest/Source/NetcodeUnitTest/Classes/Net/UnitTestChannel.h b/Engine/Plugins/NetcodeUnitTest/NetcodeUnitTest/Source/NetcodeUnitTest/Classes/Net/UnitTestChannel.h
--- a/Engine/Plugins/NetcodeUnitTest/NetcodeUnitTest/Source/NetcodeUnitTest/Classes/Net/UnitTestChannel.h
+++ b/Engine/Plugins/NetcodeUnitTest/NetcodeUnitTest/Source/NetcodeUnitTest/Classes/Net/UnitTestChannel.h
@@ -2,11 +2,15 @@
 
 #pragma once
 
-#include "CoreMinimal.h"
+// Includes
 #include "UObject/ObjectMacros.h"
 #include "Engine/Channel.h"
+#include "UObject/ObjectPtr.h"
+
 #include "UnitTestChannel.generated.h"
 
+
+// Forward declarations
 class FInBunch;
 class UNetConnection;
 class UMinimalClient;
@@ -29,7 +33,8 @@ class UUnitTestChannel : public UChannel
 
 public:
 	/** The minimal client which may require received bunch notifications */
-	UMinimalClient* MinClient;
+	UPROPERTY()
+	TObjectPtr<UMinimalClient> MinClient;
 
 	/** Whether or not this channel should verify it has been opened (resends initial packets until acked, like control channel) */
 	bool bVerifyOpen;
diff --git a/Engine/Plugins/NetcodeUnitTest/NetcodeUnitTest/Source/NetcodeUnitTest/Classes/Net/UnitTestPackageMap.h b/Engine/Plugins/NetcodeUnitTest/NetcodeUnitTest/Source/NetcodeUnitTest/Classes/Net/UnitTestPackageMap.h
--- a/Engine/Plugins/NetcodeUnitTest/NetcodeUnitTest/Source/NetcodeUnitTest/Classes/Net/UnitTestPackageMap.h
+++ b/Engine/Plugins/NetcodeUnitTest/NetcodeUnitTest/Source/NetcodeUnitTest/Classes/Net/UnitTestPackageMap.h
@@ -2,10 +2,12 @@
 
 #pragma once
 
-#include "CoreMinimal.h"
+// Includes
 #include "UObject/ObjectMacros.h"
 #include "Engine/PackageMapClient.h"
 #include "NetcodeUnitTest.h"
+#include "UObject/ObjectPtr.h"
+
 #include "UnitTestPackageMap.generated.h"
 
 
@@ -44,7 +46,8 @@ class UUnitTestPackageMap : public UPackageMapClient
 
 public:
 	/** Cached reference to the minimal client that owns this package map */
-	UMinimalClient* MinClient;
+	UPROPERTY()
+	TObjectPtr<UMinimalClient> MinClient;
 
 	/** Whether or not we are currently within execution of SerializeNewActor */
 	bool bWithinSerializeNewActor;
diff --git a/Engine/Plugins/NetcodeUnitTest/NetcodeUnitTest/Source/NetcodeUnitTest/Classes/ProcessUnitTest.h b/Engine/Plugins/NetcodeUnitTest/NetcodeUnitTest/Source/NetcodeUnitTest/Classes/ProcessUnitTest.h
--- a/Engine/Plugins/NetcodeUnitTest/NetcodeUnitTest/Source/NetcodeUnitTest/Classes/ProcessUnitTest.h
+++ b/Engine/Plugins/NetcodeUnitTest/NetcodeUnitTest/Source/NetcodeUnitTest/Classes/ProcessUnitTest.h
@@ -2,7 +2,7 @@
 
 #pragma once
 
-#include "CoreMinimal.h"
+// Includes
 #include "UObject/ObjectMacros.h"
 #include "Styling/SlateColor.h"
 #include "NetcodeUnitTest.h"
diff --git a/Engine/Plugins/NetcodeUnitTest/NetcodeUnitTest/Source/NetcodeUnitTest/Classes/UnitTask.h b/Engine/Plugins/NetcodeUnitTest/NetcodeUnitTest/Source/NetcodeUnitTest/Classes/UnitTask.h
--- a/Engine/Plugins/NetcodeUnitTest/NetcodeUnitTest/Source/NetcodeUnitTest/Classes/UnitTask.h
+++ b/Engine/Plugins/NetcodeUnitTest/NetcodeUnitTest/Source/NetcodeUnitTest/Classes/UnitTask.h
@@ -2,8 +2,9 @@
 
 #pragma once
 
-#include "CoreMinimal.h"
+// Includes
 #include "UObject/ObjectMacros.h"
+#include "UObject/ObjectPtr.h"
 
 #include "UnitLogging.h"
 
@@ -12,6 +13,7 @@
 
 // Forward declarations
 struct FMinClientParms;
+class UUnitTest;
 
 
 
@@ -85,7 +87,7 @@ int32 GetUnitTaskPriority(EUnitTaskFlags InFlags);
 UCLASS(abstract)
 class NETCODEUNITTEST_API UUnitTask : public UObject, public FUnitLogRedirect
 {
-	friend class UUnitTest;
+	friend UUnitTest;
 
 	GENERATED_UCLASS_BODY()
 
@@ -94,7 +96,8 @@ protected:
 	EUnitTaskFlags UnitTaskFlags;
 
 	/** The UnitTest that owns this task */
-	UUnitTest* Owner;
+	UPROPERTY()
+	TObjectPtr<UUnitTest> Owner;
 
 	/** Whether or not the UnitTask has started */
 	bool bStarted;
diff --git a/Engine/Plugins/NetcodeUnitTest/NetcodeUnitTest/Source/NetcodeUnitTest/Classes/UnitTest.h b/Engine/Plugins/NetcodeUnitTest/NetcodeUnitTest/Source/NetcodeUnitTest/Classes/UnitTest.h
--- a/Engine/Plugins/NetcodeUnitTest/NetcodeUnitTest/Source/NetcodeUnitTest/Classes/UnitTest.h
+++ b/Engine/Plugins/NetcodeUnitTest/NetcodeUnitTest/Source/NetcodeUnitTest/Classes/UnitTest.h
@@ -2,12 +2,13 @@
 
 #pragma once
 
-#include "CoreMinimal.h"
+// Includes
 #include "Containers/ArrayView.h"
 #include "Misc/App.h"
 #include "Misc/OutputDeviceFile.h"
 #include "Styling/SlateColor.h"
 #include "UObject/ObjectMacros.h"
+#include "UObject/ObjectPtr.h"
 
 #include "NetcodeUnitTest.h"
 #include "UnitTestBase.h"
diff --git a/Engine/Plugins/NetcodeUnitTest/NetcodeUnitTest/Source/NetcodeUnitTest/Classes/UnitTestBase.h b/Engine/Plugins/NetcodeUnitTest/NetcodeUnitTest/Source/NetcodeUnitTest/Classes/UnitTestBase.h
--- a/Engine/Plugins/NetcodeUnitTest/NetcodeUnitTest/Source/NetcodeUnitTest/Classes/UnitTestBase.h
+++ b/Engine/Plugins/NetcodeUnitTest/NetcodeUnitTest/Source/NetcodeUnitTest/Classes/UnitTestBase.h
@@ -2,7 +2,7 @@
 
 #pragma once
 
-#include "CoreMinimal.h"
+// Includes
 #include "UObject/ObjectMacros.h"
 #include "UObject/Object.h"
 
diff --git a/Engine/Plugins/NetcodeUnitTest/NetcodeUnitTest/Source/NetcodeUnitTest/Classes/UnitTestCommandlet.h b/Engine/Plugins/NetcodeUnitTest/NetcodeUnitTest/Source/NetcodeUnitTest/Classes/UnitTestCommandlet.h
--- a/Engine/Plugins/NetcodeUnitTest/NetcodeUnitTest/Source/NetcodeUnitTest/Classes/UnitTestCommandlet.h
+++ b/Engine/Plugins/NetcodeUnitTest/NetcodeUnitTest/Source/NetcodeUnitTest/Classes/UnitTestCommandlet.h
@@ -2,7 +2,7 @@
 
 #pragma once
 
-#include "CoreMinimal.h"
+// Includes
 #include "UObject/ObjectMacros.h"
 #include "Commandlets/Commandlet.h"
 
diff --git a/Engine/Plugins/NetcodeUnitTest/NetcodeUnitTest/Source/NetcodeUnitTest/Classes/UnitTestManager.h b/Engine/Plugins/NetcodeUnitTest/NetcodeUnitTest/Source/NetcodeUnitTest/Classes/UnitTestManager.h
--- a/Engine/Plugins/NetcodeUnitTest/NetcodeUnitTest/Source/NetcodeUnitTest/Classes/UnitTestManager.h
+++ b/Engine/Plugins/NetcodeUnitTest/NetcodeUnitTest/Source/NetcodeUnitTest/Classes/UnitTestManager.h
@@ -2,7 +2,7 @@
 
 #pragma once
 
-#include "CoreMinimal.h"
+// Includes
 #include "Stats/Stats.h"
 #include "UObject/ObjectMacros.h"
 #include "UObject/Object.h"
@@ -12,6 +12,7 @@
 #include "Tickable.h"
 #include "UnitTestManager.generated.h"
 
+// Forward declarations
 class FLogWindowManager;
 class SLogWindow;
 class SWindow;
diff --git a/Engine/Plugins/NetcodeUnitTest/NetcodeUnitTest/Source/NetcodeUnitTest/Classes/UnitTests/Engine/IPClient.h b/Engine/Plugins/NetcodeUnitTest/NetcodeUnitTest/Source/NetcodeUnitTest/Classes/UnitTests/Engine/IPClient.h
--- a/Engine/Plugins/NetcodeUnitTest/NetcodeUnitTest/Source/NetcodeUnitTest/Classes/UnitTests/Engine/IPClient.h
+++ b/Engine/Plugins/NetcodeUnitTest/NetcodeUnitTest/Source/NetcodeUnitTest/Classes/UnitTests/Engine/IPClient.h
@@ -2,7 +2,7 @@
 
 #pragma once
 
-#include "CoreMinimal.h"
+// Includes
 #include "UObject/ObjectMacros.h"
 #include "ClientUnitTest.h"
 
diff --git a/Engine/Plugins/NetcodeUnitTest/NetcodeUnitTest/Source/NetcodeUnitTest/Classes/UnitTests/Engine/PackedVectorTest.h b/Engine/Plugins/NetcodeUnitTest/NetcodeUnitTest/Source/NetcodeUnitTest/Classes/UnitTests/Engine/PackedVectorTest.h
--- a/Engine/Plugins/NetcodeUnitTest/NetcodeUnitTest/Source/NetcodeUnitTest/Classes/UnitTests/Engine/PackedVectorTest.h
+++ b/Engine/Plugins/NetcodeUnitTest/NetcodeUnitTest/Source/NetcodeUnitTest/Classes/UnitTests/Engine/PackedVectorTest.h
@@ -2,7 +2,7 @@
 
 #pragma once
 
-#include "CoreMinimal.h"
+// Includes
 #include "UObject/ObjectMacros.h"
 #include "UnitTest.h"
 
diff --git a/Engine/Plugins/NetcodeUnitTest/NetcodeUnitTest/Source/NetcodeUnitTest/Classes/UnitTests/VMReflection.h b/Engine/Plugins/NetcodeUnitTest/NetcodeUnitTest/Source/NetcodeUnitTest/Classes/UnitTests/VMReflection.h
--- a/Engine/Plugins/NetcodeUnitTest/NetcodeUnitTest/Source/NetcodeUnitTest/Classes/UnitTests/VMReflection.h
+++ b/Engine/Plugins/NetcodeUnitTest/NetcodeUnitTest/Source/NetcodeUnitTest/Classes/UnitTests/VMReflection.h
@@ -2,13 +2,16 @@
 
 #pragma once
 
-#include "CoreMinimal.h"
+// Includes
 #include "UObject/ObjectMacros.h"
 #include "UObject/Object.h"
 #include "UnitTest.h"
+#include "UObject/ObjectPtr.h"
 
 #include "VMReflection.generated.h"
 
+
+// Forward declarations
 class APawn;
 
 /**
@@ -34,7 +37,7 @@ class UVMTestClassA : public UObject
 
 public:
 	UPROPERTY()
-	TObjectPtr<UObject> AObjectRef;
+	TObjectPtr<UObject>	AObjectRef;
 
 	UPROPERTY()
 	uint8		ByteProp;
@@ -95,7 +98,7 @@ public:
 	uint8		BytePropArray[4];
 
 	UPROPERTY()
-	TObjectPtr<UObject> ObjectPropArray[2];
+	TObjectPtr<UObject>	ObjectPropArray[2];
 
 	UPROPERTY()
 	TArray<uint8>		DynBytePropArray;
@@ -104,7 +107,7 @@ public:
 	TArray<bool>		DynBoolPropArray;
 
 	UPROPERTY()
-	TArray<TObjectPtr<UObject>> DynObjectPropArray;
+	TArray<TObjectPtr<UObject>>	DynObjectPropArray;
 
 	UPROPERTY()
 	TArray<FName>		DynNamePropArray;
@@ -143,10 +146,10 @@ public:
 	TArray<FText>		DynTextPropArray;
 
 	UPROPERTY()
-	TArray<TObjectPtr<UClass>> DynClassPropArray;
+	TArray<TObjectPtr<UClass>>		DynClassPropArray;
 
 	UPROPERTY()
-	TArray<TObjectPtr<APawn>> DynPawnPropArray;
+	TArray<TObjectPtr<APawn>>		DynPawnPropArray;
 
 
 	UPROPERTY()
@@ -166,6 +169,6 @@ class UVMTestClassB : public UObject
 
 public:
 	UPROPERTY()
-	TObjectPtr<UObject> BObjectRef;
+	TObjectPtr<UObject>	BObjectRef;
 };
 
diff --git a/Engine/Plugins/NetcodeUnitTest/NetcodeUnitTest/Source/NetcodeUnitTest/NetcodeUnitTest.Build.cs b/Engine/Plugins/NetcodeUnitTest/NetcodeUnitTest/Source/NetcodeUnitTest/NetcodeUnitTest.Build.cs
--- a/Engine/Plugins/NetcodeUnitTest/NetcodeUnitTest/Source/NetcodeUnitTest/NetcodeUnitTest.Build.cs
+++ b/Engine/Plugins/NetcodeUnitTest/NetcodeUnitTest/Source/NetcodeUnitTest/NetcodeUnitTest.Build.cs
@@ -34,7 +34,6 @@ namespace UnrealBuildTool.Rules
 				}
 			);
 
-			// @todo #JohnBLowpri: Currently don't support standalone commandlet, with static builds (can't get past linker error in Win32)
 			if (Target.LinkType != TargetLinkType.Monolithic)
 			{
 				PrivateDependencyModuleNames.AddRange
diff --git a/Engine/Plugins/NetcodeUnitTest/NetcodeUnitTest/Source/NetcodeUnitTest/Private/ClientUnitTest.cpp b/Engine/Plugins/NetcodeUnitTest/NetcodeUnitTest/Source/NetcodeUnitTest/Private/ClientUnitTest.cpp
--- a/Engine/Plugins/NetcodeUnitTest/NetcodeUnitTest/Source/NetcodeUnitTest/Private/ClientUnitTest.cpp
+++ b/Engine/Plugins/NetcodeUnitTest/NetcodeUnitTest/Source/NetcodeUnitTest/Private/ClientUnitTest.cpp
@@ -20,9 +20,6 @@
 #include "NUTUtilDebug.h"
 #include "NUTUtilReflection.h"
 
-
-UClass* UClientUnitTest::OnlineBeaconClass = FindObject<UClass>(ANY_PACKAGE, TEXT("OnlineBeaconClient"));
-
 // @todo #JohnB: Create a unit test for load-testing servers, using multiple instances of the minimal client,
 //					as a feature for engine testing and licensees in general (and to flesh-out support for multiple min clients,
 //					in the unit test code).
@@ -68,6 +65,11 @@ UClientUnitTest::UClientUnitTest(const FObjectInitializer& ObjectInitializer)
 {
 }
 
+void UClientUnitTest::InitOnlineBeaconClass()
+{
+	OnlineBeaconClass_Private = FindObject<UClass>(ANY_PACKAGE, TEXT("OnlineBeaconClient"));
+}
+
 void UClientUnitTest::NotifyAlterMinClient(FMinClientParms& Parms)
 {
 	UUnitTask* CurTask = UnitTasks.Num() > 0 ? UnitTasks[0] : nullptr;
@@ -155,6 +157,8 @@ void UClientUnitTest::NotifyAllowNetActor(UClass* ActorClass, bool bActorChannel
 		bBlockActor = false;
 	}
 
+	const UClass* OnlineBeaconClass = GetOnlineBeaconClass();
+
 	check(OnlineBeaconClass != nullptr);
 
 	if (!!(UnitTestFlags & EUnitTestFlags::RequireBeacon) && ActorClass->IsChildOf(OnlineBeaconClass) && UnitBeacon == nullptr)
@@ -190,6 +194,8 @@ void UClientUnitTest::NotifyNetActor(UActorChannel* ActorChannel, AActor* Actor)
 		}
 	}
 
+	const UClass* OnlineBeaconClass = GetOnlineBeaconClass();
+
 	check(OnlineBeaconClass != nullptr);
 
 	if (!!(UnitTestFlags & EUnitTestFlags::BeaconConnect) && !UnitBeacon.IsValid() && Actor->IsA(OnlineBeaconClass))
diff --git a/Engine/Plugins/NetcodeUnitTest/NetcodeUnitTest/Source/NetcodeUnitTest/Private/NUTActor.cpp b/Engine/Plugins/NetcodeUnitTest/NetcodeUnitTest/Source/NetcodeUnitTest/Private/NUTActor.cpp
--- a/Engine/Plugins/NetcodeUnitTest/NetcodeUnitTest/Source/NetcodeUnitTest/Private/NUTActor.cpp
+++ b/Engine/Plugins/NetcodeUnitTest/NetcodeUnitTest/Source/NetcodeUnitTest/Private/NUTActor.cpp
@@ -304,7 +304,7 @@ bool ANUTActor::VerifyEventWatcher()
 	bool bVerified = false;
 	TObjectPtr<UNetConnection>& EventWatcher = UNUTGlobals::Get().EventWatcher;
 
-	if (EventWatcher != nullptr)
+	if (EventWatcher)
 	{
 		UWorld* CurWorld = NUTUtil::GetPrimaryWorld();
 		UNetDriver* CurDriver = (CurWorld != nullptr ? NUTUtil::GetActiveNetDriver(CurWorld) : nullptr);
diff --git a/Engine/Plugins/NetcodeUnitTest/NetcodeUnitTest/Source/NetcodeUnitTest/Private/NUTUtilProfiler.cpp b/Engine/Plugins/NetcodeUnitTest/NetcodeUnitTest/Source/NetcodeUnitTest/Private/NUTUtilProfiler.cpp
--- a/Engine/Plugins/NetcodeUnitTest/NetcodeUnitTest/Source/NetcodeUnitTest/Private/NUTUtilProfiler.cpp
+++ b/Engine/Plugins/NetcodeUnitTest/NetcodeUnitTest/Source/NetcodeUnitTest/Private/NUTUtilProfiler.cpp
@@ -4,7 +4,7 @@
 #include "Stats/Stats.h"
 
 #if STATS
-
+#include "NetcodeUnitTest.h"
 #include "Stats/StatsData.h"
 
 
diff --git a/Engine/Plugins/NetcodeUnitTest/NetcodeUnitTest/Source/NetcodeUnitTest/Private/Net/NUTUtilNet.cpp b/Engine/Plugins/NetcodeUnitTest/NetcodeUnitTest/Source/NetcodeUnitTest/Private/Net/NUTUtilNet.cpp
--- a/Engine/Plugins/NetcodeUnitTest/NetcodeUnitTest/Source/NetcodeUnitTest/Private/Net/NUTUtilNet.cpp
+++ b/Engine/Plugins/NetcodeUnitTest/NetcodeUnitTest/Source/NetcodeUnitTest/Private/Net/NUTUtilNet.cpp
@@ -243,6 +243,16 @@ void FNetworkNotifyHook::NotifyControlMessage(UNetConnection* Connection, uint8
  * FWorldTickHook
  */
 
+void FWorldTickHook::FInternalGCObject::AddReferencedObjects(FReferenceCollector& Collector)
+{
+	Collector.AddReferencedObject(Owner->AttachedWorld);
+}
+
+FString FWorldTickHook::FInternalGCObject::GetReferencerName() const
+{
+	return TEXT("FWorldTickHook::FInternalGCObject");
+}
+
 void FWorldTickHook::Init()
 {
 	if (AttachedWorld != nullptr)
diff --git a/Engine/Plugins/NetcodeUnitTest/NetcodeUnitTest/Source/NetcodeUnitTest/Private/Tests/AllUnitTestsTest.cpp b/Engine/Plugins/NetcodeUnitTest/NetcodeUnitTest/Source/NetcodeUnitTest/Private/Tests/AllUnitTestsTest.cpp
--- a/Engine/Plugins/NetcodeUnitTest/NetcodeUnitTest/Source/NetcodeUnitTest/Private/Tests/AllUnitTestsTest.cpp
+++ b/Engine/Plugins/NetcodeUnitTest/NetcodeUnitTest/Source/NetcodeUnitTest/Private/Tests/AllUnitTestsTest.cpp
@@ -1,9 +1,8 @@
 // Copyright Epic Games, Inc. All Rights Reserved.
 
-#include "CoreMinimal.h"
+// Includes
 #include "Misc/AutomationTest.h"
 
-
 #include "UnitTestManager.h"
 #include "UnitTest.h"
 #include "NUTUtil.h"
diff --git a/Engine/Plugins/NetcodeUnitTest/NetcodeUnitTest/Source/NetcodeUnitTest/Private/UI/SLogWidget.cpp b/Engine/Plugins/NetcodeUnitTest/NetcodeUnitTest/Source/NetcodeUnitTest/Private/UI/SLogWidget.cpp
--- a/Engine/Plugins/NetcodeUnitTest/NetcodeUnitTest/Source/NetcodeUnitTest/Private/UI/SLogWidget.cpp
+++ b/Engine/Plugins/NetcodeUnitTest/NetcodeUnitTest/Source/NetcodeUnitTest/Private/UI/SLogWidget.cpp
@@ -176,16 +176,12 @@ void SLogWidget::Construct(const FArguments& Args)
 	auto ConditionalSlot =
 		[] (bool bCondition, SHorizontalBox::FSlot::FSlotArguments& InSlot) -> SHorizontalBox::FSlot::FSlotArguments&
 		{
-			if (bCondition)
+			if (!bCondition)
 			{
-				return InSlot;
-			}
-			else
-			{
-				return
-					SHorizontalBox::Slot()
-					.AutoWidth();
+				InSlot.AttachWidget(SNullWidget::NullWidget);
 			}
+
+			return InSlot;
 		};
 
 
@@ -219,8 +215,6 @@ void SLogWidget::Construct(const FArguments& Args)
 					.AutoWidth()
 					[
 						SNew(SButton)
-						.ButtonStyle(FCoreStyle::Get(), "Toolbar.Button")
-						.ForegroundColor(FSlateColor::UseForeground())
 						.ToolTipText(FText::FromString(TEXT("Open the find bar for the current tab.")))
 						.OnClicked_Lambda(
 							[&]()
@@ -230,25 +224,7 @@ void SLogWidget::Construct(const FArguments& Args)
 						})
 						[
 							SNew(SImage)
-							// @todo #JohnBUI: The scaled image looks a bit fuzzy, so perhaps don't do that
-							//				(if you can figure out how to make it clip rather than scale though, do that)
-#if 1
 							.Image(&FCoreStyle::Get().GetWidgetStyle<FSearchBoxStyle>("SearchBox").GlassImage)
-#else
-							.Image_Lambda(
-								[]()
-								{
-									static FSlateBrush SearchBrush;
-
-									SearchBrush = FCoreStyle::Get().GetWidgetStyle<FSearchBoxStyle>("SearchBox").GlassImage;
-
-									// Tweak the size, so that the search button isn't oversized relative to other log window widgets
-									SearchBrush.ImageSize *= 0.8f;
-
-
-									return &SearchBrush;
-								})
-#endif
 							.ColorAndOpacity(FSlateColor::UseForeground())
 						]
 					]
diff --git a/Engine/Plugins/NetcodeUnitTest/NetcodeUnitTest/Source/NetcodeUnitTest/Private/UnitTestManager.cpp b/Engine/Plugins/NetcodeUnitTest/NetcodeUnitTest/Source/NetcodeUnitTest/Private/UnitTestManager.cpp
--- a/Engine/Plugins/NetcodeUnitTest/NetcodeUnitTest/Source/NetcodeUnitTest/Private/UnitTestManager.cpp
+++ b/Engine/Plugins/NetcodeUnitTest/NetcodeUnitTest/Source/NetcodeUnitTest/Private/UnitTestManager.cpp
@@ -521,7 +521,7 @@ bool UUnitTestManager::WithinUnitTestLimits(UClass* PendingUnitTest/*=NULL*/)
 
 		uint32 FirstRunCount = 0;
 
-		for (auto CurUnitTest : ActiveUnitTests)
+		for (UUnitTest* CurUnitTest : ActiveUnitTests)
 		{
 			if (CurUnitTest->IsFirstTimeStats())
 			{
@@ -595,7 +595,7 @@ bool UUnitTestManager::WithinUnitTestLimits(UClass* PendingUnitTest/*=NULL*/)
 				SIZE_T ReturnVal = 0;
 				double PeakMemTime = InUnitTest->StartTime + (double)InUnitTest->TimeToPeakMem;
 
-				for (auto CurUnitTest : ActiveUnitTests)
+				for (UUnitTest* CurUnitTest : ActiveUnitTests)
 				{
 					if (CurUnitTest == InUnitTest)
 					{
@@ -758,7 +758,7 @@ void UUnitTestManager::NotifyLogWindowClosed(const TSharedRef<SWindow>& ClosedWi
 				return Element->LogWindow == ClosedWindow;
 			});
 
-		UUnitTest* CurUnitTest = (CurUnitTestRef != NULL ? *CurUnitTestRef : NULL);
+		UUnitTest* CurUnitTest = (CurUnitTestRef ? CurUnitTestRef->Get() : nullptr);
 
 		if (CurUnitTest != NULL)
 		{
@@ -822,7 +822,7 @@ void UUnitTestManager::NotifyCloseAllDialogResult(const TSharedRef<SWindow>& Dia
 		// Now abort all active unit tests
 		TArray<UUnitTest*> ActiveUnitTestsCopy(ActiveUnitTests);
 
-		for (auto CurUnitTest : ActiveUnitTestsCopy)
+		for (UUnitTest* CurUnitTest : ActiveUnitTestsCopy)
 		{
 			CurUnitTest->AbortUnitTest();
 		}
@@ -866,7 +866,7 @@ void UUnitTestManager::DumpStatus(bool bForce/*=false*/)
 		STATUS_LOG(ELogType::StyleUnderline, TEXT("Unit test status:"))
 		STATUS_LOG(, TEXT("- Number of active unit tests: %i"), ActiveUnitTests.Num());
 
-		for (auto CurUnitTest : ActiveUnitTests)
+		for (UUnitTest* CurUnitTest : ActiveUnitTests)
 		{
 			TotalMemoryUsage += CurUnitTest->CurrentMemoryUsage;
 
@@ -878,7 +878,7 @@ void UUnitTestManager::DumpStatus(bool bForce/*=false*/)
 
 		STATUS_LOG(, TEXT("- Number of pending unit tests: %i"), PendingUnitTests.Num());
 
-		for (auto CurClass : PendingUnitTests)
+		for (UClass* CurClass : PendingUnitTests)
 		{
 			UUnitTest* CurUnitTest =  Cast<UUnitTest>(CurClass->GetDefaultObject());
 
@@ -1084,7 +1084,7 @@ void UUnitTestManager::PrintFinalSummary()
 	TArray<FString> AbortList;
 	TArray<UUnitTest*> UnfinishedUnitTests;
 
-	for (auto CurUnitTest : FinishedUnitTests)
+	for (UUnitTest* CurUnitTest : FinishedUnitTests)
 	{
 		if (CurUnitTest->bAborted)
 		{
@@ -1115,7 +1115,7 @@ void UUnitTestManager::PrintFinalSummary()
 
 		bool bUnitTestCompleted = false;
 
-		for (auto CurUnitTest : FinishedUnitTests)
+		for (UUnitTest* CurUnitTest : FinishedUnitTests)
 		{
 			if (CurUnitTest->GetUnitTestName() == CurAbort)
 			{
@@ -1210,12 +1210,12 @@ void UUnitTestManager::PrintFinalSummary()
 			}
 		};
 
-	for (auto CurUnitTest : FinishedUnitTests)
+	for (UUnitTest* CurUnitTest : FinishedUnitTests)
 	{
 		StatusPrintResult(CurUnitTest, false);
 	}
 
-	for (auto CurUnitTest : UnfinishedUnitTests)
+	for (UUnitTest* CurUnitTest : UnfinishedUnitTests)
 	{
 		StatusPrintResult(CurUnitTest, true);
 	}
@@ -1303,7 +1303,7 @@ void UUnitTestManager::Tick(float DeltaTime)
 		double CurTime = FPlatformTime::Seconds();
 		const double NetTickInterval = 1.0 / 60.0;
 
-		for (auto CurUnitTest : ActiveUnitTestsCopy)
+		for (UUnitTest* CurUnitTest : ActiveUnitTestsCopy)
 		{
 			if (IsValid(CurUnitTest))
 			{
@@ -1343,7 +1343,7 @@ void UUnitTestManager::Tick(float DeltaTime)
 
 
 		// Now mark all of these unit tests, for garbage collection
-		for (auto CurUnitTest : FinishedUnitTests)
+		for (UUnitTest* CurUnitTest : FinishedUnitTests)
 		{
 			CurUnitTest->MarkAsGarbage();
 		}
@@ -1392,7 +1392,7 @@ void UUnitTestManager::Tick(float DeltaTime)
 
 		for (int32 i=ActiveUnitTests.Num()-1; i>=0; i--)
 		{
-			UUnitTest* CurUnitTest = ActiveUnitTests[i];
+			UUnitTest* CurUnitTest = ActiveUnitTests[i].Get();
 			SIZE_T CurMemUsage = CurUnitTest->CurrentMemoryUsage;
 
 
@@ -1412,7 +1412,7 @@ void UUnitTestManager::Tick(float DeltaTime)
 				{
 					uint8 CloseCount = 0;
 
-					for (auto CurFinished : FinishedUnitTests)
+					for (UUnitTest* CurFinished : FinishedUnitTests)
 					{
 						if (CurFinished->bAborted && CurFinished->GetClass() == UnitTestClass)
 						{
@@ -1594,7 +1594,7 @@ bool UUnitTestManager::Exec(UWorld* InWorld, const TCHAR* Cmd, FOutputDevice& Ar
 		TObjectPtr<UUnitTest>* UnitTestRef = (InWorld != nullptr ? ActiveUnitTests.FindByPredicate(FindUnitTestByWorld) :
 									ActiveUnitTests.FindByPredicate(FindUnitTestByLog));
 
-		UUnitTest* AbortUnitTest = (UnitTestRef != nullptr ? *UnitTestRef : nullptr);
+		UUnitTest* AbortUnitTest = (UnitTestRef ? UnitTestRef->Get() : nullptr);
 
 		if (AbortUnitTest != nullptr)
 		{
@@ -1613,7 +1613,7 @@ bool UUnitTestManager::Exec(UWorld* InWorld, const TCHAR* Cmd, FOutputDevice& Ar
 		TObjectPtr<UUnitTest>* UnitTestRef = (InWorld != nullptr ? ActiveUnitTests.FindByPredicate(FindUnitTestByWorld) :
 									ActiveUnitTests.FindByPredicate(FindUnitTestByLog));
 
-		UUnitTest* RestartUnitTest = (UnitTestRef != nullptr ? *UnitTestRef : nullptr);
+		UUnitTest* RestartUnitTest = (UnitTestRef ? UnitTestRef->Get() : nullptr);
 
 		if (RestartUnitTest != nullptr)
 		{
@@ -1673,7 +1673,7 @@ bool UUnitTestManager::Exec(UWorld* InWorld, const TCHAR* Cmd, FOutputDevice& Ar
 			})
 			: nullptr);
 
-		UClientUnitTest* TargetUnitTest = (TargetUnitTestRef != nullptr ? Cast<UClientUnitTest>(*TargetUnitTestRef) : nullptr);
+		UClientUnitTest* TargetUnitTest = (TargetUnitTestRef ? Cast<UClientUnitTest>(TargetUnitTestRef->Get()) : nullptr);
 
 		// Alternatively, if a unit test has not launched started connecting to a server, its world may not be setup,
 		// so can detect by checking the active log unit test too
@@ -1937,7 +1937,7 @@ void UUnitTestManager::Serialize(const TCHAR* Data, ELogVerbosity::Type Verbosit
 
 				if (SourceInterface == nullptr)
 				{
-					for (auto CurUnitTest : ActiveUnitTests)
+					for (UUnitTest* CurUnitTest : ActiveUnitTests)
 					{
 						if (CurUnitTest->IsConnectionLogSource(GActiveReceiveUnitConnection))
 						{
@@ -1977,7 +1977,7 @@ void UUnitTestManager::Serialize(const TCHAR* Data, ELogVerbosity::Type Verbosit
 
 				if (SourceInterface == nullptr)
 				{
-					for (auto CurUnitTest : ActiveUnitTests)
+					for (UUnitTest* CurUnitTest : ActiveUnitTests)
 					{
 						UClientUnitTest* CurClientUnitTest = Cast<UClientUnitTest>(CurUnitTest);
 						UMinimalClient* CurMinClient = (CurClientUnitTest != nullptr ? CurClientUnitTest->MinClient : nullptr);
@@ -2022,7 +2022,7 @@ void UUnitTestManager::Serialize(const TCHAR* Data, ELogVerbosity::Type Verbosit
 
 						if (TimerObj != nullptr)
 						{
-							for (auto CurUnitTest : ActiveUnitTests)
+							for (UUnitTest* CurUnitTest : ActiveUnitTests)
 							{
 								if (CurUnitTest->IsTimerLogSource(TimerObj))
 								{
diff --git a/Engine/Plugins/NetcodeUnitTest/NetcodeUnitTest/Source/NetcodeUnitTest/Public/INetcodeUnitTest.h b/Engine/Plugins/NetcodeUnitTest/NetcodeUnitTest/Source/NetcodeUnitTest/Public/INetcodeUnitTest.h
--- a/Engine/Plugins/NetcodeUnitTest/NetcodeUnitTest/Source/NetcodeUnitTest/Public/INetcodeUnitTest.h
+++ b/Engine/Plugins/NetcodeUnitTest/NetcodeUnitTest/Source/NetcodeUnitTest/Public/INetcodeUnitTest.h
@@ -2,7 +2,7 @@
 
 #pragma once
 
-#include "CoreMinimal.h"
+// Includes
 #include "Modules/ModuleInterface.h"
 #include "Modules/ModuleManager.h"
 
diff --git a/Engine/Plugins/NetcodeUnitTest/NetcodeUnitTest/Source/NetcodeUnitTest/Public/NUTUtil.h b/Engine/Plugins/NetcodeUnitTest/NetcodeUnitTest/Source/NetcodeUnitTest/Public/NUTUtil.h
--- a/Engine/Plugins/NetcodeUnitTest/NetcodeUnitTest/Source/NetcodeUnitTest/Public/NUTUtil.h
+++ b/Engine/Plugins/NetcodeUnitTest/NetcodeUnitTest/Source/NetcodeUnitTest/Public/NUTUtil.h
@@ -2,7 +2,7 @@
 
 #pragma once
 
-#include "CoreMinimal.h"
+// Includes
 #include "GameFramework/Actor.h"
 #include "NetcodeUnitTest.h"
 #include "EngineGlobals.h"
@@ -10,6 +10,7 @@
 #include "Misc/OutputDeviceError.h"
 
 
+// Forward declarations
 class UNetDriver;
 class UPackage;
 struct FStackTracker;
diff --git a/Engine/Plugins/NetcodeUnitTest/NetcodeUnitTest/Source/NetcodeUnitTest/Public/NUTUtilDebug.h b/Engine/Plugins/NetcodeUnitTest/NetcodeUnitTest/Source/NetcodeUnitTest/Public/NUTUtilDebug.h
--- a/Engine/Plugins/NetcodeUnitTest/NetcodeUnitTest/Source/NetcodeUnitTest/Public/NUTUtilDebug.h
+++ b/Engine/Plugins/NetcodeUnitTest/NetcodeUnitTest/Source/NetcodeUnitTest/Public/NUTUtilDebug.h
@@ -3,16 +3,17 @@
 #pragma once
 
 // Includes
-
-#include "CoreMinimal.h"
 #include "NetcodeUnitTest.h"
 #include "NUTUtil.h"
 #include "Containers/StackTracker.h"
 
+
+// Forward declarations
 class FLogStackTraceManager;
 class FStackTraceManager;
 class UClientUnitTest;
 
+
 // Globals
 
 /** Provides a globally accessible trace manager, for easy access to stack trace debugging */
diff --git a/Engine/Plugins/NetcodeUnitTest/NetcodeUnitTest/Source/NetcodeUnitTest/Public/NUTUtilProfiler.h b/Engine/Plugins/NetcodeUnitTest/NetcodeUnitTest/Source/NetcodeUnitTest/Public/NUTUtilProfiler.h
--- a/Engine/Plugins/NetcodeUnitTest/NetcodeUnitTest/Source/NetcodeUnitTest/Public/NUTUtilProfiler.h
+++ b/Engine/Plugins/NetcodeUnitTest/NetcodeUnitTest/Source/NetcodeUnitTest/Public/NUTUtilProfiler.h
@@ -2,8 +2,11 @@
 
 #pragma once
 
-#include "CoreMinimal.h"
-#include "NetcodeUnitTest.h"
+// Includes
+#include "Misc/Build.h"
+#include "UObject/NameTypes.h"
+#include "Delegates/IDelegateInstance.h"
+
 
 #if STATS
 /**
diff --git a/Engine/Plugins/NetcodeUnitTest/NetcodeUnitTest/Source/NetcodeUnitTest/Public/NUTUtilReflection.h b/Engine/Plugins/NetcodeUnitTest/NetcodeUnitTest/Source/NetcodeUnitTest/Public/NUTUtilReflection.h
--- a/Engine/Plugins/NetcodeUnitTest/NetcodeUnitTest/Source/NetcodeUnitTest/Public/NUTUtilReflection.h
+++ b/Engine/Plugins/NetcodeUnitTest/NetcodeUnitTest/Source/NetcodeUnitTest/Public/NUTUtilReflection.h
@@ -11,6 +11,8 @@
 #include "UObject/StructOnScope.h"
 #include "UObject/UnrealType.h"
 
+
+// Forward declarations
 class FStructOnScope;
 class FFuncReflection;
 
diff --git a/Engine/Plugins/NetcodeUnitTest/NetcodeUnitTest/Source/NetcodeUnitTest/Public/Net/NUTUtilNet.h b/Engine/Plugins/NetcodeUnitTest/NetcodeUnitTest/Source/NetcodeUnitTest/Public/Net/NUTUtilNet.h
--- a/Engine/Plugins/NetcodeUnitTest/NetcodeUnitTest/Source/NetcodeUnitTest/Public/Net/NUTUtilNet.h
+++ b/Engine/Plugins/NetcodeUnitTest/NetcodeUnitTest/Source/NetcodeUnitTest/Public/Net/NUTUtilNet.h
@@ -2,16 +2,18 @@
 
 #pragma once
 
-#include "CoreMinimal.h"
+// Includes
 #include "Engine/Channel.h"
 #include "NetcodeUnitTest.h"
 #include "Engine/NetworkDelegates.h"
 #include "Engine/World.h"
 #include "Sockets.h"
+#include "UObject/GCObject.h"
 
 #include "Net/UnitTestPackageMap.h"
 
 
+// Forward declarations
 class APlayerController;
 class FInBunch;
 class FInternetAddr;
@@ -27,7 +29,6 @@ class UUnitTestNetDriver;
 class UClientUnitTest;
 
 
-
 // Delegates
 
 /**
@@ -253,9 +254,29 @@ public:
  */
 class FWorldTickHook
 {
+private:
+	// Internal GC hook, due to 'Init' function name conflict
+	class FInternalGCObject : public FGCObject
+	{
+	public:
+		FInternalGCObject() = delete;
+
+		FInternalGCObject(FWorldTickHook* InOwner)
+			: Owner(InOwner)
+		{
+		}
+
+		virtual void AddReferencedObjects(FReferenceCollector& Collector) override;
+		virtual FString GetReferencerName() const override;
+
+	private:
+		FWorldTickHook* Owner;
+	};
+
 public:
 	FWorldTickHook(UWorld* InWorld)
 		: AttachedWorld(InWorld)
+		, InternalGC(this)
 	{
 	}
 
@@ -278,6 +299,9 @@ private:
 
 	/** Handle for PostTick dispatch delegate */
 	FDelegateHandle PostTickFlushDelegateHandle;
+
+	/** Internal GC */
+	FInternalGCObject InternalGC;
 };
 
 
diff --git a/Engine/Plugins/NetcodeUnitTest/NetcodeUnitTest/Source/NetcodeUnitTest/Public/NetcodeUnitTest.h b/Engine/Plugins/NetcodeUnitTest/NetcodeUnitTest/Source/NetcodeUnitTest/Public/NetcodeUnitTest.h
--- a/Engine/Plugins/NetcodeUnitTest/NetcodeUnitTest/Source/NetcodeUnitTest/Public/NetcodeUnitTest.h
+++ b/Engine/Plugins/NetcodeUnitTest/NetcodeUnitTest/Source/NetcodeUnitTest/Public/NetcodeUnitTest.h
@@ -2,12 +2,11 @@
 
 #pragma once
 
-#include "CoreMinimal.h"
+// Includes
+#include "Logging/LogMacros.h"
 
-/**
- * Forward declarations
- */
 
+// Forward declarations
 class UUnitTestManager;
 
 
diff --git a/Engine/Plugins/NetcodeUnitTest/NetcodeUnitTest/Source/NetcodeUnitTest/Public/UI/LogWidgetCommands.h b/Engine/Plugins/NetcodeUnitTest/NetcodeUnitTest/Source/NetcodeUnitTest/Public/UI/LogWidgetCommands.h
--- a/Engine/Plugins/NetcodeUnitTest/NetcodeUnitTest/Source/NetcodeUnitTest/Public/UI/LogWidgetCommands.h
+++ b/Engine/Plugins/NetcodeUnitTest/NetcodeUnitTest/Source/NetcodeUnitTest/Public/UI/LogWidgetCommands.h
@@ -2,7 +2,7 @@
 
 #pragma once
 
-#include "CoreMinimal.h"
+// Includes
 #include "Styling/CoreStyle.h"
 #include "Framework/Commands/Commands.h"
 
diff --git a/Engine/Plugins/NetcodeUnitTest/NetcodeUnitTest/Source/NetcodeUnitTest/Public/UI/LogWindowManager.h b/Engine/Plugins/NetcodeUnitTest/NetcodeUnitTest/Source/NetcodeUnitTest/Public/UI/LogWindowManager.h
--- a/Engine/Plugins/NetcodeUnitTest/NetcodeUnitTest/Source/NetcodeUnitTest/Public/UI/LogWindowManager.h
+++ b/Engine/Plugins/NetcodeUnitTest/NetcodeUnitTest/Source/NetcodeUnitTest/Public/UI/LogWindowManager.h
@@ -2,8 +2,7 @@
 
 #pragma once
 
-#include "CoreMinimal.h"
-
+// Includes
 #include "NetcodeUnitTest.h"
 #include "UnitLogging.h"
 
diff --git a/Engine/Plugins/NetcodeUnitTest/NetcodeUnitTest/Source/NetcodeUnitTest/Public/UI/SLogDialog.h b/Engine/Plugins/NetcodeUnitTest/NetcodeUnitTest/Source/NetcodeUnitTest/Public/UI/SLogDialog.h
--- a/Engine/Plugins/NetcodeUnitTest/NetcodeUnitTest/Source/NetcodeUnitTest/Public/UI/SLogDialog.h
+++ b/Engine/Plugins/NetcodeUnitTest/NetcodeUnitTest/Source/NetcodeUnitTest/Public/UI/SLogDialog.h
@@ -2,7 +2,12 @@
 
 #pragma once
 
-#include "CoreMinimal.h"
+// Includes
+#include "Delegates/DelegateCombinations.h"
+#include "Templates/SharedPointer.h"
+#include "UObject/NoExportTypes.h"
+#include "Internationalization/Text.h"
+#include "Containers/UnrealString.h"
 
 class SWindow;
 
diff --git a/Engine/Plugins/NetcodeUnitTest/NetcodeUnitTest/Source/NetcodeUnitTest/Public/UI/SLogWidget.h b/Engine/Plugins/NetcodeUnitTest/NetcodeUnitTest/Source/NetcodeUnitTest/Public/UI/SLogWidget.h
--- a/Engine/Plugins/NetcodeUnitTest/NetcodeUnitTest/Source/NetcodeUnitTest/Public/UI/SLogWidget.h
+++ b/Engine/Plugins/NetcodeUnitTest/NetcodeUnitTest/Source/NetcodeUnitTest/Public/UI/SLogWidget.h
@@ -2,7 +2,7 @@
 
 #pragma once
 
-#include "CoreMinimal.h"
+// Includes
 #include "Misc/CommandLine.h"
 #include "Input/Reply.h"
 #include "Styling/SlateColor.h"
@@ -15,6 +15,7 @@
 #include "UnitLogging.h"
 
 
+// Forward declarations
 class FUICommandList;
 class SEditableTextBox;
 enum class ESuspendState : uint8;
diff --git a/Engine/Plugins/NetcodeUnitTest/NetcodeUnitTest/Source/NetcodeUnitTest/Public/UI/SLogWindow.h b/Engine/Plugins/NetcodeUnitTest/NetcodeUnitTest/Source/NetcodeUnitTest/Public/UI/SLogWindow.h
--- a/Engine/Plugins/NetcodeUnitTest/NetcodeUnitTest/Source/NetcodeUnitTest/Public/UI/SLogWindow.h
+++ b/Engine/Plugins/NetcodeUnitTest/NetcodeUnitTest/Source/NetcodeUnitTest/Public/UI/SLogWindow.h
@@ -2,7 +2,7 @@
 
 #pragma once
 
-#include "CoreMinimal.h"
+// Includes
 #include "Widgets/DeclarativeSyntaxSupport.h"
 #include "Input/Reply.h"
 #include "Widgets/SWindow.h"
@@ -11,6 +11,7 @@
 #include "UnitLogging.h"
 
 
+// Forward declarations
 class SLogWidget;
 
 /**
diff --git a/Engine/Plugins/NetcodeUnitTest/NetcodeUnitTest/Source/NetcodeUnitTest/Public/UI/SMultiSelectTableRow.h b/Engine/Plugins/NetcodeUnitTest/NetcodeUnitTest/Source/NetcodeUnitTest/Public/UI/SMultiSelectTableRow.h
--- a/Engine/Plugins/NetcodeUnitTest/NetcodeUnitTest/Source/NetcodeUnitTest/Public/UI/SMultiSelectTableRow.h
+++ b/Engine/Plugins/NetcodeUnitTest/NetcodeUnitTest/Source/NetcodeUnitTest/Public/UI/SMultiSelectTableRow.h
@@ -2,13 +2,14 @@
 
 #pragma once
 
-#include "CoreMinimal.h"
+// Includes
 #include "Layout/Geometry.h"
 #include "Input/DragAndDrop.h"
 #include "Input/Reply.h"
 #include "Widgets/Views/STableRow.h"
 #include "NetcodeUnitTest.h"
 
+
 /**
  * SMultiSelectTableRow
  *
diff --git a/Engine/Plugins/NetcodeUnitTest/NetcodeUnitTest/Source/NetcodeUnitTest/Public/UnitLogging.h b/Engine/Plugins/NetcodeUnitTest/NetcodeUnitTest/Source/NetcodeUnitTest/Public/UnitLogging.h
--- a/Engine/Plugins/NetcodeUnitTest/NetcodeUnitTest/Source/NetcodeUnitTest/Public/UnitLogging.h
+++ b/Engine/Plugins/NetcodeUnitTest/NetcodeUnitTest/Source/NetcodeUnitTest/Public/UnitLogging.h
@@ -2,7 +2,7 @@
 
 #pragma once
 
-#include "CoreMinimal.h"
+// Includes
 #include "Misc/EnumClassFlags.h"
 #include "Styling/SlateColor.h"
 
diff --git a/Engine/Plugins/NetcodeUnitTest/NetcodeUnitTest/Source/NetcodeUnitTest/Public/UnitTestEnvironment.h b/Engine/Plugins/NetcodeUnitTest/NetcodeUnitTest/Source/NetcodeUnitTest/Public/UnitTestEnvironment.h
--- a/Engine/Plugins/NetcodeUnitTest/NetcodeUnitTest/Source/NetcodeUnitTest/Public/UnitTestEnvironment.h
+++ b/Engine/Plugins/NetcodeUnitTest/NetcodeUnitTest/Source/NetcodeUnitTest/Public/UnitTestEnvironment.h
@@ -2,13 +2,21 @@
 
 #pragma once
 
-#include "CoreMinimal.h"
+// Includes
+#include "Containers/UnrealString.h"
+#include "Containers/Array.h"
+#include "NUTEnum.h"
 
+
+// Forward declarations
 class APlayerController;
 class FUnitTestEnvironment;
 enum class EUnitTestFlags : uint32;
 class UUnitTest;
 
+
+// Globals
+
 /** The list of registered unit test environments, mapped to game name (unit test modules should add to this, within StartupModule) */
 extern TMap<FString, FUnitTestEnvironment*> UnitTestEnvironments;
 
