Added '-LogDebug=MatchStr' commandline parameter and 'LogTrace AddDebug MatchStr' console command to NetcodeUnitTest, to trigger debugging on a particular log line pattern.

diff --git a/Engine/Plugins/NetcodeUnitTest/NetcodeUnitTest/Source/NetcodeUnitTest/Private/NUTUtilDebug.cpp b/Engine/Plugins/NetcodeUnitTest/NetcodeUnitTest/Source/NetcodeUnitTest/Private/NUTUtilDebug.cpp
--- a/Engine/Plugins/NetcodeUnitTest/NetcodeUnitTest/Source/NetcodeUnitTest/Private/NUTUtilDebug.cpp
+++ b/Engine/Plugins/NetcodeUnitTest/NetcodeUnitTest/Source/NetcodeUnitTest/Private/NUTUtilDebug.cpp
@@ -67,7 +67,7 @@ void FScopedLog::InternalConstruct(const TArray<FString>& InLogCategories, UClie
 			ENUTControlCommand ControlCmd = ENUTControlCommand::Command_NoResult;
 			FString Cmd = TEXT("");
 
-			for (auto CurCategory : LogCategories)
+			for (const FString& CurCategory : LogCategories)
 			{
 				Cmd = TEXT("Log ") + CurCategory + TEXT(" ") + TargetVerbosity;
 
@@ -195,7 +195,7 @@ bool FProcessEventHookBase::InternalProcessEventHook(AActor* Actor, UFunction* F
 void FScopedProcessEventLog::ProcessEventHook(AActor* Actor, UFunction* Function, void* Parameters)
 {
 	UE_LOG(LogUnitTest, Log, TEXT("FScopedProcessEventLog: Actor: %s, Function: %s"),
-			(Actor != NULL ? *Actor->GetName() : TEXT("NULL")),
+			(Actor != nullptr ? *Actor->GetName() : TEXT("nullptr")),
 			*Function->GetName());
 }
 #endif
@@ -269,11 +269,11 @@ FStackTraceManager::FStackTraceManager()
 
 FStackTraceManager::~FStackTraceManager()
 {
-	for (auto It=Traces.CreateConstIterator(); It; ++It)
+	for (TMap<FString, FNUTStackTrace*>::TConstIterator It=Traces.CreateConstIterator(); It; ++It)
 	{
 		FNUTStackTrace* Trace = It->Value;
 
-		if (Trace != NULL)
+		if (Trace != nullptr)
 		{
 			delete Trace;
 		}
@@ -293,7 +293,7 @@ void FStackTraceManager::Disable(FString TraceName)
 {
 	FNUTStackTrace* Trace = GetTrace(TraceName);
 
-	if (Trace != NULL)
+	if (Trace != nullptr)
 	{
 		Trace->Disable();
 	}
@@ -335,7 +335,7 @@ void FStackTraceManager::Dump(FString TraceName, bool bKeepTraceHistory/*=false*
 {
 	FNUTStackTrace* Trace = GetTrace(TraceName);
 
-	if (Trace != NULL)
+	if (Trace != nullptr)
 	{
 		Trace->Dump(bKeepTraceHistory);
 
@@ -355,7 +355,7 @@ void FStackTraceManager::Clear(FString TraceName)
 {
 	FNUTStackTrace* Trace = GetTrace(TraceName);
 
-	if (Trace != NULL)
+	if (Trace != nullptr)
 	{
 		delete Trace;
 		Traces.FindAndRemoveChecked(TraceName);
@@ -370,11 +370,11 @@ void FStackTraceManager::DumpAll(bool bKeepTraceHistory/*=false*/, bool bKeepTra
 {
 	UE_LOG(LogUnitTest, Log, TEXT("Dumping all tracked stack traces:"));
 
-	for (auto It=Traces.CreateIterator(); It; ++It)
+	for (TMap<FString, FNUTStackTrace*>::TIterator It=Traces.CreateIterator(); It; ++It)
 	{
 		FNUTStackTrace* Trace = It->Value;
 
-		if (Trace != NULL)
+		if (Trace != nullptr)
 		{
 			Trace->Dump(bKeepTraceHistory);
 
@@ -391,7 +391,7 @@ void FStackTraceManager::TraceAndDump(FString TraceName)
 {
 	FNUTStackTrace* Trace = GetTrace(TraceName);
 
-	if (Trace == NULL || Trace->IsTrackingEnabled())
+	if (Trace == nullptr || Trace->IsTrackingEnabled())
 	{
 		UE_LOG(LogUnitTest, Log, TEXT("Dumping once-off stack trace for TraceName '%s':"), *TraceName);
 
@@ -404,6 +404,201 @@ void FStackTraceManager::TraceAndDump(FString TraceName)
 }
 
 
+/**
+ * FLogStackTraceManager
+ */
+
+FLogStackTraceManager::~FLogStackTraceManager()
+{
+	if (GLog != nullptr)
+	{
+		GLog->RemoveOutputDevice(this);
+	}
+}
+
+void FLogStackTraceManager::AddLogTrace(FString LogLine, ELogTraceFlags TraceFlags)
+{
+	const bool bPartial = EnumHasAnyFlags(TraceFlags, ELogTraceFlags::Partial);
+
+	if (!bPartial && !EnumHasAnyFlags(TraceFlags, ELogTraceFlags::Full))
+	{
+		UE_LOG(LogUnitTest, Warning, TEXT("AddLogTrace: Should specify either ELogTraceFlags::Full or ELogTraceFlags::Partial. Enabling 'Full'."));
+
+		TraceFlags |= ELogTraceFlags::Full;
+	}
+	else if (bPartial && EnumHasAnyFlags(TraceFlags, ELogTraceFlags::Full))
+	{
+		UE_LOG(LogUnitTest, Warning, TEXT("AddLogTrace: Specified both ELogTraceFlags::Full and ELogTraceFlags::Partial. Defaulting to 'Partial'."));
+
+		TraceFlags &= ~ELogTraceFlags::Full;
+	}
+
+	const bool bDumpTrace = EnumHasAnyFlags(TraceFlags, ELogTraceFlags::DumpTrace);
+	const bool bDebugTrace = EnumHasAnyFlags(TraceFlags, ELogTraceFlags::Debug);
+
+	UE_LOG(LogUnitTest, Log, TEXT("Adding %slog trace for line: %s (DumpTrace: %i, Debug: %i)"),
+			(bPartial ? TEXT("partial ") : TEXT("")), *LogLine, (int32)bDumpTrace, (int32)bDebugTrace);
+
+	// Add the log hook, if not active already
+	if (!GLog->IsRedirectingTo(this))
+	{
+		GLog->AddOutputDevice(this);
+	}
+
+	FLogTraceEntry CurEntry;
+
+	CurEntry.LogLine = LogLine;
+	CurEntry.TraceFlags = TraceFlags;
+
+
+	auto MatchLogLine =
+		[&](const FLogTraceEntry& InEntry)
+		{
+			return InEntry.LogLine == LogLine;
+		};
+
+	if (bPartial)
+	{
+		if (!PartialMatches.ContainsByPredicate(MatchLogLine))
+		{
+			PartialMatches.Add(CurEntry);
+		}
+	}
+	else
+	{
+		if (!ExactMatches.ContainsByPredicate(MatchLogLine))
+		{
+			ExactMatches.Add(CurEntry);
+		}
+	}
+}
+
+void FLogStackTraceManager::ClearLogTrace(FString LogLine, bool bDump/*=true*/)
+{
+	UE_LOG(LogUnitTest, Log, TEXT("Clearing log trace for line: %s"), *LogLine);
+
+	auto MatchLogLine =
+		[&](const FLogTraceEntry& CurEntry)
+		{
+			return CurEntry.LogLine == LogLine;
+		};
+
+	int32 ExactIdx = ExactMatches.IndexOfByPredicate(MatchLogLine);
+	int32 PartialIdx = PartialMatches.IndexOfByPredicate(MatchLogLine);
+
+	if (ExactIdx != INDEX_NONE)
+	{
+		ExactMatches.RemoveAt(ExactIdx);
+	}
+
+	if (PartialIdx != INDEX_NONE)
+	{
+		PartialMatches.RemoveAt(PartialIdx);
+	}
+
+
+	if (PartialIdx != INDEX_NONE || ExactIdx != INDEX_NONE)
+	{
+		if (bDump)
+		{
+			GTraceManager->Dump(LogLine, false, false);
+		}
+		else
+		{
+			GTraceManager->Clear(LogLine);
+		}
+	}
+
+	if (PartialMatches.Num() == 0 && ExactMatches.Num() == 0)
+	{
+		GLog->RemoveOutputDevice(this);
+	}
+}
+
+void FLogStackTraceManager::ClearAll(bool bDump/*=false*/)
+{
+	UE_LOG(LogUnitTest, Log, TEXT("Clearing all log traces."));
+
+	if (bDump)
+	{
+		for (const FLogTraceEntry& CurEntry : ExactMatches)
+		{
+			if (GTraceManager->ContainsTrace(CurEntry.LogLine))
+			{
+				GTraceManager->Dump(CurEntry.LogLine, false, false);
+			}
+			else
+			{
+				UE_LOG(LogUnitTest, Log, TEXT("No stack traces for log trace: %s"), *CurEntry.LogLine);
+			}
+		}
+
+		for (auto CurEntry : PartialMatches)
+		{
+			if (GTraceManager->ContainsTrace(CurEntry.LogLine))
+			{
+				GTraceManager->Dump(CurEntry.LogLine, false, false);
+			}
+			else
+			{
+				UE_LOG(LogUnitTest, Log, TEXT("No stack traces for (partial) log trace: %s"), *CurEntry.LogLine);
+			}
+		}
+	}
+
+	ExactMatches.Empty();
+	PartialMatches.Empty();
+
+	GLog->RemoveOutputDevice(this);
+}
+
+void FLogStackTraceManager::Serialize(const TCHAR* Data, ELogVerbosity::Type Verbosity, const class FName& Category)
+{
+	static bool bWithinLogTrace = false;
+
+	if (!bWithinLogTrace)
+	{
+		auto CheckCurEntry = [](const FLogTraceEntry& CurEntry)
+		{
+			bWithinLogTrace = true;
+
+			if (EnumHasAnyFlags(CurEntry.TraceFlags, ELogTraceFlags::DumpTrace))
+			{
+				GTraceManager->TraceAndDump(CurEntry.LogLine);
+			}
+			else
+			{
+				GTraceManager->AddTrace(CurEntry.LogLine);
+			}
+
+			if (EnumHasAnyFlags(CurEntry.TraceFlags, ELogTraceFlags::Debug))
+			{
+				// Careful - you're about to crash unless you hit 'Continue' before you hit 'Stop Debugging'
+				UE_DEBUG_BREAK();
+			}
+
+			bWithinLogTrace = false;
+		};
+
+		for (const FLogTraceEntry& CurEntry : ExactMatches)
+		{
+			if (CurEntry.LogLine == Data)
+			{
+				CheckCurEntry(CurEntry);
+			}
+		}
+
+		for (const FLogTraceEntry& CurEntry : PartialMatches)
+		{
+			if (FCString::Stristr(Data, *CurEntry.LogLine) != nullptr)
+			{
+				CheckCurEntry(CurEntry);
+			}
+		}
+	}
+}
+
+
 /**
  * NUTDebug
  */
diff --git a/Engine/Plugins/NetcodeUnitTest/NetcodeUnitTest/Source/NetcodeUnitTest/Private/NetcodeUnitTest.cpp b/Engine/Plugins/NetcodeUnitTest/NetcodeUnitTest/Source/NetcodeUnitTest/Private/NetcodeUnitTest.cpp
--- a/Engine/Plugins/NetcodeUnitTest/NetcodeUnitTest/Source/NetcodeUnitTest/Private/NetcodeUnitTest.cpp
+++ b/Engine/Plugins/NetcodeUnitTest/NetcodeUnitTest/Source/NetcodeUnitTest/Private/NetcodeUnitTest.cpp
@@ -101,7 +101,12 @@ public:
 
 		if (FParse::Value(FCommandLine::Get(), TEXT("LogTrace="), TraceStr) && TraceStr.Len() > 0)
 		{
-			GLogTraceManager->AddLogTrace(TraceStr);
+			GLogTraceManager->AddLogTrace(TraceStr, ELogTraceFlags::Partial | ELogTraceFlags::DumpTrace);
+		}
+
+		if (FParse::Value(FCommandLine::Get(), TEXT("LogDebug="), TraceStr) && TraceStr.Len() > 0)
+		{
+			GLogTraceManager->AddLogTrace(TraceStr, ELogTraceFlags::Partial | ELogTraceFlags::Debug);
 		}
 	}
 
diff --git a/Engine/Plugins/NetcodeUnitTest/NetcodeUnitTest/Source/NetcodeUnitTest/Private/UnitTestManager.cpp b/Engine/Plugins/NetcodeUnitTest/NetcodeUnitTest/Source/NetcodeUnitTest/Private/UnitTestManager.cpp
--- a/Engine/Plugins/NetcodeUnitTest/NetcodeUnitTest/Source/NetcodeUnitTest/Private/UnitTestManager.cpp
+++ b/Engine/Plugins/NetcodeUnitTest/NetcodeUnitTest/Source/NetcodeUnitTest/Private/UnitTestManager.cpp
@@ -2200,6 +2200,9 @@ static bool UnitTestExec(UWorld* InWorld, const TCHAR* Cmd, FOutputDevice& Ar)
 	 *	- Add a partial log line to tracking (case insensitive, and can match substrings):
 	 *		GEngine->Exec(NULL, TEXT("LogTrace AddPartial LogLine"));
 	 *
+	 *	- Add a partial log line to tracking for debug breaking:
+	 *		GEngine->Exec(NULL, TEXT("LogTrace AddDebug LogLine"));
+	 *
 	 *	- Dump accumulated log entries, for a specified log line, and clears it from tracing:
 	 *		GEngine->Exec(NULL, TEXT("LogTrace Dump LogLine"));
 	 *
@@ -2216,11 +2219,15 @@ static bool UnitTestExec(UWorld* InWorld, const TCHAR* Cmd, FOutputDevice& Ar)
 
 		if (FParse::Command(&Cmd, TEXT("Add")) && (LogLine = Cmd).Len() > 0)
 		{
-			GLogTraceManager->AddLogTrace(LogLine, false);
+			GLogTraceManager->AddLogTrace(LogLine, ELogTraceFlags::Full);
 		}
 		else if (FParse::Command(&Cmd, TEXT("AddPartial")) && (LogLine = Cmd).Len() > 0)
 		{
-			GLogTraceManager->AddLogTrace(LogLine, true);
+			GLogTraceManager->AddLogTrace(LogLine, ELogTraceFlags::Partial);
+		}
+		else if (FParse::Command(&Cmd, TEXT("AddDebug")) && (LogLine = Cmd).Len() > 0)
+		{
+			GLogTraceManager->AddLogTrace(LogLine, ELogTraceFlags::Partial | ELogTraceFlags::Debug);
 		}
 		else if (FParse::Command(&Cmd, TEXT("Dump")) && (LogLine = Cmd).Len() > 0)
 		{
diff --git a/Engine/Plugins/NetcodeUnitTest/NetcodeUnitTest/Source/NetcodeUnitTest/Public/NUTUtilDebug.h b/Engine/Plugins/NetcodeUnitTest/NetcodeUnitTest/Source/NetcodeUnitTest/Public/NUTUtilDebug.h
--- a/Engine/Plugins/NetcodeUnitTest/NetcodeUnitTest/Source/NetcodeUnitTest/Public/NUTUtilDebug.h
+++ b/Engine/Plugins/NetcodeUnitTest/NetcodeUnitTest/Source/NetcodeUnitTest/Public/NUTUtilDebug.h
@@ -433,8 +433,18 @@ protected:
 //				some stuff (like enable/disable) should be easy to add though, so do that eventually
 
 
-// @todo #JohnB: NOTE: Don't move the log stack trace manager to .cpp yet; not until you decide how to implement,
-//				the similar ProcessEvent hook that will tie in to the trace manager, as you may try to generalize the below for that
+/**
+ * Flags specifying the type of log trace
+ */
+enum class ELogTraceFlags : uint16
+{
+	Full			= 0x0001,	// Full string match log trace (excluding category and verbosity level)
+	Partial			= 0x0002,	// Partial string match log trace
+	DumpTrace		= 0x0004,	// Dump a stack trace every time there's a match
+	Debug			= 0x0008	// Start debugging on match (if debugger attached)
+};
+
+ENUM_CLASS_FLAGS(ELogTraceFlags);
 
 /**
  * A log hook, which watches the log for specified log entries, and ties them into the stack trace manager
@@ -452,14 +462,8 @@ public:
 		/**	The log line being to watch for */
 		FString LogLine;
 
-		/** Whether or not to do a dump every time the log entry is encountered */
-		bool bDump;
-
-		FLogTraceEntry()
-			: LogLine()
-			, bDump(false)
-		{
-		}
+		/** Flags specifying the type of log trace */
+		ELogTraceFlags TraceFlags;
 	};
 
 public:
@@ -473,13 +477,7 @@ public:
 	/**
 	 * Destructor
 	 */
-	virtual ~FLogStackTraceManager() override
-	{
-		if (GLog != nullptr)
-		{
-			GLog->RemoveOutputDevice(this);
-		}
-	}
+	virtual ~FLogStackTraceManager() override;
 
 	/**
 	 * Adds a log line for log trace tracking
@@ -490,133 +488,43 @@ public:
 	 * @param bPartial	Whether or not to do partial matches for this line (case insensitive, and matches substrings)
 	 * @param bDump		Whether or not to dump traces as they are encountered
 	 */
+	UE_DEPRECATED(5.2, "Use AddLogTrace with ELogTraceFlags now, instead.")
 	void AddLogTrace(FString LogLine, bool bPartial=true, bool bDump=true)
 	{
-		UE_LOG(LogUnitTest, Log, TEXT("Adding %slog trace for line: %s"), (bPartial ? TEXT("partial ") : TEXT("")), *LogLine);
+		ELogTraceFlags TraceFlags = (bPartial ? ELogTraceFlags::Partial : ELogTraceFlags::Full);
 
-		// Add the log hook, if not active already
-		if (!GLog->IsRedirectingTo(this))
+		if (bDump)
 		{
-			GLog->AddOutputDevice(this);
+			TraceFlags |= ELogTraceFlags::DumpTrace;
 		}
 
-		FLogTraceEntry CurEntry;
-
-		CurEntry.LogLine = LogLine;
-		CurEntry.bDump = bDump;
-
-
-		auto MatchLogLine =
-			[&](const FLogTraceEntry& InEntry)
-			{
-				return InEntry.LogLine == LogLine;
-			};
-
-		if (bPartial)
-		{
-			if (!PartialMatches.ContainsByPredicate(MatchLogLine))
-			{
-				PartialMatches.Add(CurEntry);
-			}
-		}
-		else
-		{
-			if (!ExactMatches.ContainsByPredicate(MatchLogLine))
-			{
-				ExactMatches.Add(CurEntry);
-			}
-		}
+		AddLogTrace(LogLine, TraceFlags);
 	}
 
+	/**
+	 * Adds a log line for log trace tracking
+	 * NOTE: The LogLine does NOT match the category or verbosity (e.g. LogNet or Warning/Verbose) of logs
+	 * NOTE: Partial matches, will output to log when encountered, so matched logs can be identified
+	 *
+	 * @param LogLine		The line to be tracked
+	 * @param TraceFlags	Flags for the type of trace to perform
+	 */
+	void AddLogTrace(FString LogLine, ELogTraceFlags TraceFlags);
+
 	/**
 	 * Removes a log line from trace tracking
 	 *
 	 * @param LogLine	The log line that was being tracked
 	 * @param bDump		Whether or not to dump from the trace manager as well (defaults to true); cleared from trace manager either way
 	 */
-	void ClearLogTrace(FString LogLine, bool bDump=true)
-	{
-		UE_LOG(LogUnitTest, Log, TEXT("Clearing log trace for line: %s"), *LogLine);
-
-		auto MatchLogLine =
-			[&](const FLogTraceEntry& CurEntry)
-			{
-				return CurEntry.LogLine == LogLine;
-			};
-
-		int32 ExactIdx = ExactMatches.IndexOfByPredicate(MatchLogLine);
-		int32 PartialIdx = PartialMatches.IndexOfByPredicate(MatchLogLine);
-
-		if (ExactIdx != INDEX_NONE)
-		{
-			ExactMatches.RemoveAt(ExactIdx);
-		}
-
-		if (PartialIdx != INDEX_NONE)
-		{
-			PartialMatches.RemoveAt(PartialIdx);
-		}
-
-
-		if (PartialIdx != INDEX_NONE || ExactIdx != INDEX_NONE)
-		{
-			if (bDump)
-			{
-				GTraceManager->Dump(LogLine, false, false);
-			}
-			else
-			{
-				GTraceManager->Clear(LogLine);
-			}
-		}
-
-		if (PartialMatches.Num() == 0 && ExactMatches.Num() == 0)
-		{
-			GLog->RemoveOutputDevice(this);
-		}
-	}
+	void ClearLogTrace(FString LogLine, bool bDump=true);
 
 	/**
 	 * Clears all log tracing
 	 *
 	 * @param bDump		Whether or not to dump all cleared log trace entries
 	 */
-	void ClearAll(bool bDump=false)
-	{
-		UE_LOG(LogUnitTest, Log, TEXT("Clearing all log traces."));
-
-		if (bDump)
-		{
-			for (auto CurEntry : ExactMatches)
-			{
-				if (GTraceManager->ContainsTrace(CurEntry.LogLine))
-				{
-					GTraceManager->Dump(CurEntry.LogLine, false, false);
-				}
-				else
-				{
-					UE_LOG(LogUnitTest, Log, TEXT("No stack traces for log trace: %s"), *CurEntry.LogLine);
-				}
-			}
-
-			for (auto CurEntry : PartialMatches)
-			{
-				if (GTraceManager->ContainsTrace(CurEntry.LogLine))
-				{
-					GTraceManager->Dump(CurEntry.LogLine, false, false);
-				}
-				else
-				{
-					UE_LOG(LogUnitTest, Log, TEXT("No stack traces for (partial) log trace: %s"), *CurEntry.LogLine);
-				}
-			}
-		}
-
-		ExactMatches.Empty();
-		PartialMatches.Empty();
-
-		GLog->RemoveOutputDevice(this);
-	}
+	void ClearAll(bool bDump=false);
 
 
 	// We're hiding UObject::Serialize() by declaring this.  That's OK, but Clang will warn about it.
@@ -631,55 +539,7 @@ public:
 	 * @param Verbosity	The log level verbosity for the message
 	 * @param Category	The log category for the message
 	 */
-	virtual void Serialize(const TCHAR* Data, ELogVerbosity::Type Verbosity, const class FName& Category) override
-	{
-		static bool bWithinLogTrace = false;
-
-		if (!bWithinLogTrace)
-		{
-			for (auto CurEntry : ExactMatches)
-			{
-				if (CurEntry.LogLine == Data)
-				{
-					bWithinLogTrace = true;
-
-					if (CurEntry.bDump)
-					{
-						GTraceManager->TraceAndDump(CurEntry.LogLine);
-					}
-					else
-					{
-						GTraceManager->AddTrace(CurEntry.LogLine);
-					}
-
-					bWithinLogTrace = false;
-				}
-			}
-
-			if (PartialMatches.Num() > 0)
-			{
-				for (auto CurEntry : PartialMatches)
-				{
-					if (FCString::Stristr(Data, *CurEntry.LogLine) != NULL)
-					{
-						// NOTE: Do not use Data for the TraceName, as this makes things much harder to track
-						bWithinLogTrace = true;
-
-						if (CurEntry.bDump)
-						{
-							GTraceManager->TraceAndDump(CurEntry.LogLine);
-						}
-						else
-						{
-							GTraceManager->AddTrace(CurEntry.LogLine);
-						}
-
-						bWithinLogTrace = false;
-					}
-				}
-			}
-		}
-	}
+	virtual void Serialize(const TCHAR* Data, ELogVerbosity::Type Verbosity, const class FName& Category) override;
 #if defined(__clang__)
 #pragma clang diagnostic pop
 #endif
@@ -694,7 +554,7 @@ public:
 };
 
 
-// @todo #JohnBFeatureDebug: Continue implementing
+// @todo #JohnBFeatureDebug: Continue implementing (will likely tie into the FLogStackTraceManager tracing features, so refactor/generalize that)
 #if 0
 #if !UE_BUILD_SHIPPING
 /**
