Refactored log tracing and added ability to trigger console commands on certain logs.

Commandline example:
-LogCommand="ReadContentBlockPayload failed=obj list class=ShooterGameState"

Command example:
LogCommand Add ReadContentBlockPayload failed=obj list class=ShooterGameState

diff --git a/Engine/Plugins/NetcodeUnitTest/NetcodeUnitTest/Source/NetcodeUnitTest/Private/NUTUtilDebug.cpp b/Engine/Plugins/NetcodeUnitTest/NetcodeUnitTest/Source/NetcodeUnitTest/Private/NUTUtilDebug.cpp
--- a/Engine/Plugins/NetcodeUnitTest/NetcodeUnitTest/Source/NetcodeUnitTest/Private/NUTUtilDebug.cpp
+++ b/Engine/Plugins/NetcodeUnitTest/NetcodeUnitTest/Source/NetcodeUnitTest/Private/NUTUtilDebug.cpp
@@ -17,7 +17,17 @@
  */
 
 FStackTraceManager* GTraceManager = new FStackTraceManager();
+
+PRAGMA_DISABLE_DEPRECATION_WARNINGS
 FLogStackTraceManager* GLogTraceManager = new FLogStackTraceManager();
+PRAGMA_ENABLE_DEPRECATION_WARNINGS
+
+namespace UE::NUT
+{
+	FLogHookManager* GLogHookManager = new FLogHookManager();
+	FLogTraceManager* GLogTrace = new FLogTraceManager();
+	FLogCommandManager* GLogCommandManager = new FLogCommandManager();
+}
 
 
 bool GGlobalExec(UWorld* InWorld, const TCHAR* Cmd, FOutputDevice& Ar)
@@ -241,7 +251,7 @@ void FNUTStackTrace::AddTrace(bool bLogAdd/*=false*/)
 			UE_LOG(LogUnitTest, Log, TEXT("Adding stack trace for TraceName '%s'."), *TraceName);
 		}
 
-		Tracker.CaptureStackTrace(TRACE_IGNORE_DEPTH);
+		Tracker.CaptureStackTrace(TRACE_IGNORE_DEPTH + IgnoreDepthOffset);
 	}
 }
 
@@ -397,17 +407,359 @@ void FStackTraceManager::TraceAndDump(FString TraceName)
 
 		FStackTracker TempTracker(nullptr, nullptr, nullptr, true);
 
-		TempTracker.CaptureStackTrace(TRACE_IGNORE_DEPTH);
+		TempTracker.CaptureStackTrace(TRACE_IGNORE_DEPTH + IgnoreDepthOffset);
 		TempTracker.DumpStackTraces(0, *GLog);
 		TempTracker.ResetTracking();
 	}
 }
 
 
+namespace UE::NUT
+{
+
+/**
+ * FScopedIncrementTraceIgnoreDepth
+ */
+
+FScopedIncrementTraceIgnoreDepth::FScopedIncrementTraceIgnoreDepth(FStackTraceManager* InManager, int32 InIgnoreIncrement)
+	: Manager(InManager)
+	, IgnoreIncrement(InIgnoreIncrement)
+{
+	if (Manager != nullptr)
+	{
+		Manager->IgnoreDepthOffset += IgnoreIncrement;
+
+		for (TMap<FString, FNUTStackTrace*>::TConstIterator It(Manager->Traces); It; ++It)
+		{
+			It.Value()->IgnoreDepthOffset += IgnoreIncrement;
+		}
+	}
+}
+
+FScopedIncrementTraceIgnoreDepth::~FScopedIncrementTraceIgnoreDepth()
+{
+	if (Manager != nullptr)
+	{
+		Manager->IgnoreDepthOffset -= IgnoreIncrement;
+
+		for (TMap<FString, FNUTStackTrace*>::TConstIterator It(Manager->Traces); It; ++It)
+		{
+			It.Value()->IgnoreDepthOffset -= IgnoreIncrement;
+		}
+	}
+}
+
+
+/**
+ * FLogHookManager
+ */
+
+FLogHookManager::~FLogHookManager()
+{
+	DisableLogHook();
+}
+
+FLogHookID FLogHookManager::AddLogHook(FLogHook&& InHookFunc, FString InMatchStr, ELogHookType InHookType/*=ELogHookType::Partial*/)
+{
+	FLogHookID HookID = NextHookID++;
+
+	if (LogHooks.Num() == 0)
+	{
+		EnableLogHook();
+	}
+
+	LogHooks.Add(FLogHookEntry{MoveTemp(InHookFunc), InMatchStr, InHookType, HookID});
+
+	return HookID;
+}
+
+void FLogHookManager::RemoveLogHook(FLogHookID InHookID)
+{
+	for (int32 i=0; i<LogHooks.Num(); i++)
+	{
+		if (LogHooks[i].HookID == InHookID)
+		{
+			LogHooks.RemoveAt(InHookID);
+		}
+	}
+
+	if (LogHooks.Num() == 0)
+	{
+		DisableLogHook();
+	}
+}
+
+void FLogHookManager::Serialize(const TCHAR* Data, ELogVerbosity::Type Verbosity, const class FName& Category)
+{
+	if (!bWithinLogSerialize)
+	{
+		TGuardValue<bool> SerializeGuard(bWithinLogSerialize, true);
+
+		for (const FLogHookEntry& CurEntry : LogHooks)
+		{
+			if (CurEntry.HookType == ELogHookType::Full)
+			{
+				if (CurEntry.LogMatch == Data)
+				{
+					CurEntry.HookFunc(Data, Verbosity, Category);
+				}
+			}
+			else //if (CurEntry.HookType == ELogHookType::Partial)
+			{
+				if (FCString::Stristr(Data, *CurEntry.LogMatch) != nullptr)
+				{
+					CurEntry.HookFunc(Data, Verbosity, Category);
+				}
+			}
+		}
+	}
+}
+
+void FLogHookManager::EnableLogHook()
+{
+	if (!GLog->IsRedirectingTo(this))
+	{
+		GLog->AddOutputDevice(this);
+	}
+}
+
+void FLogHookManager::DisableLogHook()
+{
+	if (GLog != nullptr)
+	{
+		GLog->RemoveOutputDevice(this);
+	}
+}
+
+
+/**
+ * FLogTraceManager
+ */
+
+FLogTraceManager::~FLogTraceManager()
+{
+	ClearHooks();
+}
+
+void FLogTraceManager::AddLogTrace(FString LogLine, ELogTraceFlags TraceFlags)
+{
+	using namespace UE::NUT;
+
+	const bool bPartial = EnumHasAnyFlags(TraceFlags, ELogTraceFlags::Partial);
+
+	if (!bPartial && !EnumHasAnyFlags(TraceFlags, ELogTraceFlags::Full))
+	{
+		UE_LOG(LogUnitTest, Warning, TEXT("AddLogTrace: Should specify either ELogTraceFlags::Full or ELogTraceFlags::Partial. Enabling 'Full'."));
+
+		TraceFlags |= ELogTraceFlags::Full;
+	}
+	else if (bPartial && EnumHasAnyFlags(TraceFlags, ELogTraceFlags::Full))
+	{
+		UE_LOG(LogUnitTest, Warning, TEXT("AddLogTrace: Specified both ELogTraceFlags::Full and ELogTraceFlags::Partial. Defaulting to 'Partial'."));
+
+		TraceFlags &= ~ELogTraceFlags::Full;
+	}
+
+	const bool bDumpTrace = EnumHasAnyFlags(TraceFlags, ELogTraceFlags::DumpTrace);
+	const bool bDebugTrace = EnumHasAnyFlags(TraceFlags, ELogTraceFlags::Debug);
+
+	UE_LOG(LogUnitTest, Log, TEXT("Adding %slog trace for line: %s (DumpTrace: %i, Debug: %i)"),
+			(bPartial ? TEXT("partial ") : TEXT("")), ToCStr(LogLine), (int32)bDumpTrace, (int32)bDebugTrace);
+
+
+	if (GLogHookManager != nullptr)
+	{
+		FLogHook HookFunc = [LogLine, TraceFlags](const TCHAR* Data, ELogVerbosity::Type Verbosity, const class FName& Category)
+			{
+				FScopedIncrementTraceIgnoreDepth ScopedTraceIgnore(GTraceManager, 2);
+
+				if (EnumHasAnyFlags(TraceFlags, ELogTraceFlags::DumpTrace))
+				{
+					GTraceManager->TraceAndDump(LogLine);
+				}
+				else
+				{
+					GTraceManager->AddTrace(LogLine);
+				}
+
+				if (EnumHasAnyFlags(TraceFlags, ELogTraceFlags::Debug))
+				{
+					// Careful - you're about to crash unless you hit 'Continue' before you hit 'Stop Debugging'
+					UE_DEBUG_BREAK();
+				}
+			};
+
+		FLogHookID HookID = GLogHookManager->AddLogHook(MoveTemp(HookFunc), LogLine, (bPartial ? ELogHookType::Partial : ELogHookType::Full));
+
+		HookIDs.Add(HookID);
+
+		if (!HookMap.Contains(LogLine))
+		{
+			HookMap.Add(LogLine, HookID);
+		}
+	}
+	else
+	{
+		UE_LOG(LogUnitTest, Warning, TEXT("Error adding log trace, GLogHookManager is not set."));
+	}
+}
+
+void FLogTraceManager::ClearLogTrace(FString LogLine, bool bDump/*=true*/)
+{
+	UE_LOG(LogUnitTest, Log, TEXT("Clearing log trace for line: %s"), ToCStr(LogLine));
+
+	FLogHookID HookID = INDEX_NONE;
+
+	if (HookMap.RemoveAndCopyValue(LogLine, HookID))
+	{
+		HookIDs.Remove(HookID);
+		GLogHookManager->RemoveLogHook(HookID);
+
+		if (bDump)
+		{
+			GTraceManager->Dump(LogLine, false, false);
+		}
+		else
+		{
+			GTraceManager->Clear(LogLine);
+		}
+	}
+}
+
+void FLogTraceManager::ClearAll(bool bDump/*=false*/)
+{
+	UE_LOG(LogUnitTest, Log, TEXT("Clearing all log traces."));
+
+	if (bDump)
+	{
+		for (TMap<FString, FLogHookID>::TConstIterator It(HookMap); It; ++It)
+		{
+			const FString& CurLogLine = It.Key();
+
+			if (GTraceManager->ContainsTrace(CurLogLine))
+			{
+				GTraceManager->Dump(CurLogLine, false, false);
+			}
+			else
+			{
+				UE_LOG(LogUnitTest, Log, TEXT("No stack traces for log trace: %s"), ToCStr(CurLogLine));
+			}
+		}
+	}
+
+	ClearHooks();
+}
+
+void FLogTraceManager::ClearHooks()
+{
+	if (GLogHookManager != nullptr)
+	{
+		for (FLogHookID HookID : HookIDs)
+		{
+			GLogHookManager->RemoveLogHook(HookID);
+		}
+	}
+
+	HookIDs.Empty();
+	HookMap.Empty();
+}
+
+
+/**
+ * FLogCommandManager
+ */
+
+FLogCommandManager::~FLogCommandManager()
+{
+	ClearHooks();
+}
+
+void FLogCommandManager::AddLogCommand(FString LogLine, FString Command)
+{
+	if (GLogHookManager != nullptr)
+	{
+		UE_LOG(LogUnitTest, Log, TEXT("Adding log command '%s' for line: %s"), ToCStr(Command), ToCStr(LogLine));
+
+		FLogHook HookFunc = [Command](const TCHAR* Data, ELogVerbosity::Type Verbosity, const class FName& Category)
+			{
+				GEngine->Exec(NUTUtil::GetPrimaryWorld(), *Command);
+			};
+
+		FLogCommandEntry NewEntry;
+
+		NewEntry.LogLine = LogLine;
+		NewEntry.Command = Command;
+		NewEntry.HookID = GLogHookManager->AddLogHook(MoveTemp(HookFunc), LogLine);
+
+		LogCommands.Add(MoveTemp(NewEntry));
+	}
+	else
+	{
+		UE_LOG(LogUnitTest, Warning, TEXT("Error adding log command, GLogHookManager is not set."));
+	}
+}
+
+void FLogCommandManager::RemoveByLog(FString LogLine)
+{
+	for (int32 i=0; i<LogCommands.Num(); i++)
+	{
+		const FLogCommandEntry& CurEntry = LogCommands[i];
+
+		if (CurEntry.LogLine == LogLine)
+		{
+			if (GLogHookManager != nullptr)
+			{
+				GLogHookManager->RemoveLogHook(CurEntry.HookID);
+			}
+
+			LogCommands.RemoveAt(i);
+
+			break;
+		}
+	}
+}
+
+void FLogCommandManager::RemoveByCommand(FString Command)
+{
+	for (int32 i=0; i<LogCommands.Num(); i++)
+	{
+		const FLogCommandEntry& CurEntry = LogCommands[i];
+
+		if (CurEntry.Command == Command)
+		{
+			if (GLogHookManager != nullptr)
+			{
+				GLogHookManager->RemoveLogHook(CurEntry.HookID);
+			}
+
+			LogCommands.RemoveAt(i);
+
+			break;
+		}
+	}
+}
+
+void FLogCommandManager::ClearHooks()
+{
+	if (GLogHookManager != nullptr)
+	{
+		for (const FLogCommandEntry& CurEntry : LogCommands)
+		{
+			GLogHookManager->RemoveLogHook(CurEntry.HookID);
+		}
+	}
+
+	LogCommands.Empty();
+}
+
+}
+
 /**
  * FLogStackTraceManager
  */
 
+#if 1
+PRAGMA_DISABLE_DEPRECATION_WARNINGS
 FLogStackTraceManager::~FLogStackTraceManager()
 {
 	if (GLog != nullptr)
@@ -597,6 +949,8 @@ void FLogStackTraceManager::Serialize(const TCHAR* Data, ELogVerbosity::Type Ver
 		}
 	}
 }
+PRAGMA_ENABLE_DEPRECATION_WARNINGS
+#endif
 
 
 /**
diff --git a/Engine/Plugins/NetcodeUnitTest/NetcodeUnitTest/Source/NetcodeUnitTest/Private/NetcodeUnitTest.cpp b/Engine/Plugins/NetcodeUnitTest/NetcodeUnitTest/Source/NetcodeUnitTest/Private/NetcodeUnitTest.cpp
--- a/Engine/Plugins/NetcodeUnitTest/NetcodeUnitTest/Source/NetcodeUnitTest/Private/NetcodeUnitTest.cpp
+++ b/Engine/Plugins/NetcodeUnitTest/NetcodeUnitTest/Source/NetcodeUnitTest/Private/NetcodeUnitTest.cpp
@@ -57,6 +57,8 @@ public:
 	 */
 	virtual void StartupModule() override
 	{
+		using namespace UE::NUT;
+
 		FNUTModuleInterface::StartupModule();
 
 		static bool bSetDelegate = false;
@@ -101,12 +103,27 @@ public:
 
 		if (FParse::Value(FCommandLine::Get(), TEXT("LogTrace="), TraceStr) && TraceStr.Len() > 0)
 		{
-			GLogTraceManager->AddLogTrace(TraceStr, ELogTraceFlags::Partial | ELogTraceFlags::DumpTrace);
+			GLogTrace->AddLogTrace(TraceStr, ELogTraceFlags::Partial | ELogTraceFlags::DumpTrace);
 		}
 
 		if (FParse::Value(FCommandLine::Get(), TEXT("LogDebug="), TraceStr) && TraceStr.Len() > 0)
 		{
-			GLogTraceManager->AddLogTrace(TraceStr, ELogTraceFlags::Partial | ELogTraceFlags::Debug);
+			GLogTrace->AddLogTrace(TraceStr, ELogTraceFlags::Partial | ELogTraceFlags::Debug);
+		}
+
+		if (FParse::Value(FCommandLine::Get(), TEXT("LogCommand="), TraceStr) && TraceStr.Len() > 0)
+		{
+			FString LogLine;
+			FString Command;
+
+			if (TraceStr.Split(TEXT("="), &LogLine, &Command))
+			{
+				GLogCommandManager->AddLogCommand(LogLine, Command);
+			}
+			else
+			{
+				UE_LOG(LogUnitTest, Warning, TEXT("LogCommand commandline parameter requires the format: -LogCommand=\"LogLine=Command\""));
+			}
 		}
 	}
 
@@ -115,6 +132,8 @@ public:
 	 */
 	virtual void ShutdownModule() override
 	{
+		using namespace UE::NUT;
+
 		// Eliminate active global variables
 		GUnitTestManager = nullptr;
 
@@ -125,12 +144,35 @@ public:
 			GTraceManager = nullptr;
 		}
 
+		if (GLogHookManager != nullptr)
+		{
+			delete GLogHookManager;
+
+			GLogHookManager = nullptr;
+		}
+
+		if (GLogTrace != nullptr)
+		{
+			delete GLogTrace;
+
+			GLogTrace = nullptr;
+		}
+
+		if (GLogCommandManager != nullptr)
+		{
+			delete GLogCommandManager;
+
+			GLogCommandManager = nullptr;
+		}
+
+		PRAGMA_DISABLE_DEPRECATION_WARNINGS
 		if (GLogTraceManager != nullptr)
 		{
 			delete GLogTraceManager;
 
 			GLogTraceManager = nullptr;
 		}
+		PRAGMA_ENABLE_DEPRECATION_WARNINGS
 		
 
 		FLogWidgetCommands::Unregister();
diff --git a/Engine/Plugins/NetcodeUnitTest/NetcodeUnitTest/Source/NetcodeUnitTest/Private/UnitTestManager.cpp b/Engine/Plugins/NetcodeUnitTest/NetcodeUnitTest/Source/NetcodeUnitTest/Private/UnitTestManager.cpp
--- a/Engine/Plugins/NetcodeUnitTest/NetcodeUnitTest/Source/NetcodeUnitTest/Private/UnitTestManager.cpp
+++ b/Engine/Plugins/NetcodeUnitTest/NetcodeUnitTest/Source/NetcodeUnitTest/Private/UnitTestManager.cpp
@@ -2073,6 +2073,8 @@ void UUnitTestManager::Serialize(const TCHAR* Data, ELogVerbosity::Type Verbosit
 
 static bool UnitTestExec(UWorld* InWorld, const TCHAR* Cmd, FOutputDevice& Ar)
 {
+	using namespace UE::NUT;
+
 	bool bReturnVal = false;
 
 	if (FParse::Command(&Cmd, TEXT("UnitTest")))
@@ -2195,23 +2197,23 @@ static bool UnitTestExec(UWorld* InWorld, const TCHAR* Cmd, FOutputDevice& Ar)
 	 * Usage: (NOTE: Replace 'LogLine' with the log text to be tracked)
 	 *
 	 *	- Add an exact log line to tracking (case sensitive, and must match length too):
-	 *		GEngine->Exec(NULL, TEXT("LogTrace Add LogLine"));
+	 *		GEngine->Exec(nullptr, TEXT("LogTrace Add LogLine"));
 	 *
 	 *	- Add a partial log line to tracking (case insensitive, and can match substrings):
-	 *		GEngine->Exec(NULL, TEXT("LogTrace AddPartial LogLine"));
+	 *		GEngine->Exec(nullptr, TEXT("LogTrace AddPartial LogLine"));
 	 *
 	 *	- Add a partial log line to tracking for debug breaking:
-	 *		GEngine->Exec(NULL, TEXT("LogTrace AddDebug LogLine"));
+	 *		GEngine->Exec(nullptr, TEXT("LogTrace AddDebug LogLine"));
 	 *
 	 *	- Dump accumulated log entries, for a specified log line, and clears it from tracing:
-	 *		GEngine->Exec(NULL, TEXT("LogTrace Dump LogLine"));
+	 *		GEngine->Exec(nullptr, TEXT("LogTrace Dump LogLine"));
 	 *
 	 *	- Clear the specified log line from tracing:
-	 *		GEngine->Exec(NULL, TEXT("LogTrace Clear LogLine"));
+	 *		GEngine->Exec(nullptr, TEXT("LogTrace Clear LogLine"));
 	 *
 	 *	- Clear all log lines from tracing:
-	 *		GEngine->Exec(NULL, TEXT("LogTrace ClearAll"));
-	 *		GEngine->Exec(NULL, TEXT("LogTrace ClearAll -Dump"));
+	 *		GEngine->Exec(nullptr, TEXT("LogTrace ClearAll"));
+	 *		GEngine->Exec(nullptr, TEXT("LogTrace ClearAll -Dump"));
 	 */
 	else if (FParse::Command(&Cmd, TEXT("LogTrace")))
 	{
@@ -2219,29 +2221,29 @@ static bool UnitTestExec(UWorld* InWorld, const TCHAR* Cmd, FOutputDevice& Ar)
 
 		if (FParse::Command(&Cmd, TEXT("Add")) && (LogLine = Cmd).Len() > 0)
 		{
-			GLogTraceManager->AddLogTrace(LogLine, ELogTraceFlags::Full);
+			GLogTrace->AddLogTrace(LogLine, ELogTraceFlags::Full);
 		}
 		else if (FParse::Command(&Cmd, TEXT("AddPartial")) && (LogLine = Cmd).Len() > 0)
 		{
-			GLogTraceManager->AddLogTrace(LogLine, ELogTraceFlags::Partial);
+			GLogTrace->AddLogTrace(LogLine, ELogTraceFlags::Partial);
 		}
 		else if (FParse::Command(&Cmd, TEXT("AddDebug")) && (LogLine = Cmd).Len() > 0)
 		{
-			GLogTraceManager->AddLogTrace(LogLine, ELogTraceFlags::Partial | ELogTraceFlags::Debug);
+			GLogTrace->AddLogTrace(LogLine, ELogTraceFlags::Partial | ELogTraceFlags::Debug);
 		}
 		else if (FParse::Command(&Cmd, TEXT("Dump")) && (LogLine = Cmd).Len() > 0)
 		{
-			GLogTraceManager->ClearLogTrace(LogLine, true);
+			GLogTrace->ClearLogTrace(LogLine, true);
 		}
 		else if (FParse::Command(&Cmd, TEXT("Clear")) && (LogLine = Cmd).Len() > 0)
 		{
-			GLogTraceManager->ClearLogTrace(LogLine, false);
+			GLogTrace->ClearLogTrace(LogLine, false);
 		}
 		else if (FParse::Command(&Cmd, TEXT("ClearAll")))
 		{
 			bool bDump = FParse::Param(Cmd, TEXT("DUMP"));
 
-			GLogTraceManager->ClearAll(bDump);
+			GLogTrace->ClearAll(bDump);
 		}
 		// If LogLine is now zero-length instead of 'NotSet', that means a valid command was encountered, but no LogLine specified
 		else if (LogLine.Len() == 0)
@@ -2251,6 +2253,44 @@ static bool UnitTestExec(UWorld* InWorld, const TCHAR* Cmd, FOutputDevice& Ar)
 
 		bReturnVal = true;
 	}
+	/**
+	 * Special 'LogCommand' command, which triggers the specified console command every time a matching log entry is encountered.
+	 *
+	 * NOTE: Does not track the category or verbosity of log entries
+	 */
+	else if (FParse::Command(&Cmd, TEXT("LogCommand")))
+	{
+		FString Params;
+
+		if (FParse::Command(&Cmd, TEXT("Add")) && (Params = Cmd).Len() > 0)
+		{
+			FString LogLine;
+			FString Command;
+
+			if (Params.Split(TEXT("="), &LogLine, &Command))
+			{
+				GLogCommandManager->AddLogCommand(LogLine, Command);
+			}
+			else
+			{
+				Ar.Logf(TEXT("LogCommand requires the format: LogCommand LogLine=Command"));
+			}
+		}
+		else if (FParse::Command(&Cmd, TEXT("Remove")) && (Params = Cmd).Len() > 0)
+		{
+			GLogCommandManager->RemoveByLog(Params);
+		}
+		else if (FParse::Command(&Cmd, TEXT("RemoveCommand")) && (Params = Cmd).Len() > 0)
+		{
+			GLogCommandManager->RemoveByCommand(Params);
+		}
+		else if (Params.Len() == 0)
+		{
+			Ar.Logf(TEXT("Not enough parameters for command."));
+		}
+
+		bReturnVal = true;
+	}
 	/**
 	 * Special 'LogHex' command, for taking an arbitrary memory location, plus its length, and spitting out a hex byte-dump.
 	 * Access implemented through a console command, so that this can be used without a dependency, throughout the engine.
diff --git a/Engine/Plugins/NetcodeUnitTest/NetcodeUnitTest/Source/NetcodeUnitTest/Public/NUTUtil.h b/Engine/Plugins/NetcodeUnitTest/NetcodeUnitTest/Source/NetcodeUnitTest/Public/NUTUtil.h
--- a/Engine/Plugins/NetcodeUnitTest/NetcodeUnitTest/Source/NetcodeUnitTest/Public/NUTUtil.h
+++ b/Engine/Plugins/NetcodeUnitTest/NetcodeUnitTest/Source/NetcodeUnitTest/Public/NUTUtil.h
@@ -365,7 +365,7 @@ struct NUTUtil
 	static inline UNetDriver* GetActiveNetDriver(UWorld* InWorld)
 	{
 		UNetDriver* ReturnVal = NULL;
-		FWorldContext* WorldContext = &GEngine->GetWorldContextFromWorldChecked(InWorld); //-V595
+		FWorldContext* WorldContext = GEngine->GetWorldContextFromWorld(InWorld); //-V595
 
 		if (WorldContext != NULL && WorldContext->PendingNetGame != NULL && WorldContext->PendingNetGame->NetDriver != NULL)
 		{
diff --git a/Engine/Plugins/NetcodeUnitTest/NetcodeUnitTest/Source/NetcodeUnitTest/Public/NUTUtilDebug.h b/Engine/Plugins/NetcodeUnitTest/NetcodeUnitTest/Source/NetcodeUnitTest/Public/NUTUtilDebug.h
--- a/Engine/Plugins/NetcodeUnitTest/NetcodeUnitTest/Source/NetcodeUnitTest/Public/NUTUtilDebug.h
+++ b/Engine/Plugins/NetcodeUnitTest/NetcodeUnitTest/Source/NetcodeUnitTest/Public/NUTUtilDebug.h
@@ -12,6 +12,15 @@
 class FLogStackTraceManager;
 class FStackTraceManager;
 class UClientUnitTest;
+enum class ELogTraceFlags : uint16;
+
+namespace UE::NUT
+{
+	class FScopedIncrementTraceIgnoreDepth;
+	class FLogHookManager;
+	class FLogTraceManager;
+	class FLogCommandManager;
+}
 
 
 // Globals
@@ -19,8 +28,22 @@ class UClientUnitTest;
 /** Provides a globally accessible trace manager, for easy access to stack trace debugging */
 extern FStackTraceManager* GTraceManager;
 
+PRAGMA_DISABLE_DEPRECATION_WARNINGS
 /** Log hook for managing tying of log entry detection to the trace manager */
 extern FLogStackTraceManager* GLogTraceManager;
+PRAGMA_ENABLE_DEPRECATION_WARNINGS
+
+namespace UE::NUT
+{
+	/** Generic string matching log hook, for executing code upon certain log entries */
+	extern UE::NUT::FLogHookManager* GLogHookManager;
+
+	/** Log hook for managing tying of log entry detection to the trace manager */
+	extern UE::NUT::FLogTraceManager* GLogTrace;
+
+	/** Log hook for executing console commands for certain log entries */
+	extern UE::NUT::FLogCommandManager* GLogCommandManager;
+}
 
 
 // Enable access to FStackTracker.bIsEnabled
@@ -231,6 +254,9 @@ public:
  */
 class FNUTStackTrace
 {
+	friend FStackTraceManager;
+	friend UE::NUT::FScopedIncrementTraceIgnoreDepth;
+
 public:
 	/**
 	 * Constructs the debug stack trace
@@ -284,6 +310,10 @@ protected:
 
 	/** The stack tracker associated with this debug trace */
 	FStackTracker Tracker;
+
+private:
+	/** Additional increment for the stack trace ignore depth, as set by FScopedIncrementTraceIgnoreDepth */
+	int32 IgnoreDepthOffset;
 };
 
 
@@ -302,6 +332,8 @@ protected:
  */
 class FStackTraceManager
 {
+	friend UE::NUT::FScopedIncrementTraceIgnoreDepth;
+
 public:
 	/**
 	 * Base constructor
@@ -412,7 +444,10 @@ protected:
 
 		if (ReturnVal == NULL)
 		{
-			ReturnVal = Traces.Add(TraceName, new FNUTStackTrace(TraceName));
+			FNUTStackTrace* NewTrace = new FNUTStackTrace(TraceName);
+
+			NewTrace->IgnoreDepthOffset = IgnoreDepthOffset;
+			ReturnVal = Traces.Add(TraceName, NewTrace);
 
 			if (bOutNewTrace)
 			{
@@ -426,11 +461,224 @@ protected:
 protected:
 	/** A map of active debug stack traces */
 	TMap<FString, FNUTStackTrace*> Traces;
+
+private:
+	/** Additional increment for the stack trace ignore depth, as set by FScopedIncrementTraceIgnoreDepth */
+	int32 IgnoreDepthOffset;
 };
 
 
-// @todo #JohnBDebug: The log trace stuff doesn't support the 'startenabled' and enable/disable etc. features;
-//				some stuff (like enable/disable) should be easy to add though, so do that eventually
+namespace UE::NUT
+{
+
+/**
+ * Used for adjusting the depth of ignore stack trace lines in FStackTraceManager
+ */
+class FScopedIncrementTraceIgnoreDepth final
+{
+public:
+	/**
+	 * Base constructor
+	 *
+	 * @param InManager				The stack trace manager whose ignore depth is being adjusted
+	 * @param InIgnoreIncrement		The amount by which the stack trace managers ignore deoth should be incremented
+	 */
+	FScopedIncrementTraceIgnoreDepth(FStackTraceManager* InManager, int32 InIgnoreIncrement);
+
+	~FScopedIncrementTraceIgnoreDepth();
+
+private:
+	FStackTraceManager* Manager = nullptr;
+
+	int32 IgnoreIncrement = 0;
+};
+
+/**
+ * Specifies the type of log hook (e.g. type of string matching to use)
+ */
+enum class ELogHookType : uint8
+{
+	Full			= 0x0001,	// Full string match against log line (excluding category and verbosity level)
+	Partial			= 0x0002	// Partial string match against log line
+};
+
+/** Function to call as the log hook */
+using FLogHook = TUniqueFunction<void(const TCHAR* /*Data*/, ELogVerbosity::Type /*Verbosity*/, const class FName& /*Category*/)>;
+
+/** Log hook ID, for removing later if desired */
+using FLogHookID = int32;
+
+/**
+ * Generic log hook, for calling a hook function upon detecting a specific log entry
+ */
+class FLogHookManager final : protected FOutputDevice
+{
+private:
+	struct FLogHookEntry
+	{
+		/** The log hook function to execute */
+		FLogHook HookFunc;
+
+		/**	The log string to match on */
+		FString LogMatch;
+
+		/** Flags specifying how to match against the log line */
+		ELogHookType HookType;
+
+		/** Simple ID for identifying the entry later */
+		FLogHookID HookID;
+	};
+
+public:
+	FLogHookManager() = default;
+
+	virtual ~FLogHookManager() override;
+
+	/**
+	 * Adds a log hook, with the specified matching
+	 *
+	 * @param InHookFunc	The log hook function to call upon match
+	 * @param InMatchStr	The string to use for matching against log entries
+	 * @param InHookType	The type of matching to use for the match string
+	 * @return				Returns a unique id for the newly added log hook entry, which should be used to remove it later
+	 */
+	FLogHookID AddLogHook(FLogHook&& InHookFunc, FString InMatchStr, ELogHookType InHookType=ELogHookType::Partial);
+
+	/**
+	 * Removes a log hook using its unique id
+	 *
+	 * @param InHookID		The unique id for the log hook entry
+	 */
+	void RemoveLogHook(FLogHookID InHookID);
+
+private:
+#if defined(__clang__)
+#pragma clang diagnostic push
+#pragma clang diagnostic ignored "-Woverloaded-virtual"
+#endif
+	/**
+	 * Detects log entries
+	 *
+	 * @param Data		The text to write
+	 * @param Verbosity	The log level verbosity for the message
+	 * @param Category	The log category for the message
+	 */
+	virtual void Serialize(const TCHAR* Data, ELogVerbosity::Type Verbosity, const class FName& Category) override;
+#if defined(__clang__)
+#pragma clang diagnostic pop
+#endif
+
+	void EnableLogHook();
+
+	void DisableLogHook();
+
+private:
+	/** List of active log hooks */
+	TArray<FLogHookEntry> LogHooks;
+
+	/** Incrementing log hook ID */
+	FLogHookID NextHookID = 0;
+
+	/** Guard bool for preventing recursive calls to Serialize */
+	bool bWithinLogSerialize = false;
+};
+
+/**
+ * A log hook, which watches the log for specified log entries, and ties them into the stack trace manager
+ *
+ * Most easily used through the LogTrace console command, as documented in UnitTestManager.cpp
+ */
+class FLogTraceManager final
+{
+public:
+	FLogTraceManager() = default;
+
+	~FLogTraceManager();
+
+	/**
+	 * Adds a log line for log trace tracking
+	 * NOTE: The LogLine does NOT match the category or verbosity (e.g. LogNet or Warning/Verbose) of logs
+	 * NOTE: Partial matches, will output to log when encountered, so matched logs can be identified
+	 *
+	 * @param LogLine		The line to be tracked
+	 * @param TraceFlags	Flags for the type of trace to perform
+	 */
+	void AddLogTrace(FString LogLine, ELogTraceFlags TraceFlags);
+
+	/**
+	 * Removes a log line from trace tracking
+	 *
+	 * @param LogLine	The log line that was being tracked
+	 * @param bDump		Whether or not to dump from the trace manager as well (defaults to true); cleared from trace manager either way
+	 */
+	void ClearLogTrace(FString LogLine, bool bDump=true);
+
+	/**
+	 * Clears all log tracing
+	 *
+	 * @param bDump		Whether or not to dump all cleared log trace entries
+	 */
+	void ClearAll(bool bDump=false);
+
+private:
+	void ClearHooks();
+
+private:
+	/** Active log hook ID's, for later cleanup */
+	TArray<FLogHookID> HookIDs;
+
+	/** Active log matching strings and their ID, for tracking and cleanup */
+	TMap<FString, FLogHookID> HookMap;
+};
+
+/**
+ * Log hook which executes console commands upon detecting certain log entries.
+ *
+ * Used by the LogCommand console command and commandline parameters.
+ */
+class FLogCommandManager final
+{
+public:
+	FLogCommandManager() = default;
+
+	~FLogCommandManager();
+
+	/**
+	 * Adds a new log command
+	 *
+	 * @param LogLine	The log line to do a partial match on
+	 * @param Command	The console command to execute
+	 */
+	void AddLogCommand(FString LogLine, FString Command);
+
+	/**
+	 * Removes a log command based on the log line used for matching
+	 *
+	 * @param LogLine	The log line used for matching
+	 */
+	void RemoveByLog(FString LogLine);
+
+	/**
+	 * Removes a log command based on the console command it executes
+	 */
+	void RemoveByCommand(FString Command);
+
+private:
+	void ClearHooks();
+
+
+private:
+	struct FLogCommandEntry
+	{
+		FString		LogLine;
+		FString		Command;
+		FLogHookID	HookID = INDEX_NONE;
+	};
+
+	TArray<FLogCommandEntry> LogCommands;
+};
+
+}
 
 
 /**
@@ -446,12 +694,14 @@ enum class ELogTraceFlags : uint16
 
 ENUM_CLASS_FLAGS(ELogTraceFlags);
 
+
 /**
  * A log hook, which watches the log for specified log entries, and ties them into the stack trace manager
  *
  * Most easily used through the LogTrace console command, as documented in UnitTestManager.cpp
  */
-class FLogStackTraceManager : public FOutputDevice
+class UE_DEPRECATED(5.2, "FLogStackTraceManager is deprecated, use FLogTraceManager/GLogTrace now, instead.") FLogStackTraceManager
+	: public FOutputDevice
 {
 public:
 	/**
@@ -467,16 +717,8 @@ public:
 	};
 
 public:
-	/**
-	 * Base constructor
-	 */
-	FLogStackTraceManager()
-	{
-	}
+	FLogStackTraceManager() = default;
 
-	/**
-	 * Destructor
-	 */
 	virtual ~FLogStackTraceManager() override;
 
 	/**
@@ -554,7 +796,7 @@ public:
 };
 
 
-// @todo #JohnBFeatureDebug: Continue implementing (will likely tie into the FLogStackTraceManager tracing features, so refactor/generalize that)
+// @todo #JohnBFeatureDebug: Continue implementing (will likely tie into the FLogTraceManager tracing features, so refactor/generalize that)
 #if 0
 #if !UE_BUILD_SHIPPING
 /**
