Inclusivity.

diff --git a/Engine/Plugins/NetcodeUnitTest/NetcodeUnitTest/Source/NetcodeUnitTest/Private/UI/SLogWidget.cpp b/Engine/Plugins/NetcodeUnitTest/NetcodeUnitTest/Source/NetcodeUnitTest/Private/UI/SLogWidget.cpp
--- a/Engine/Plugins/NetcodeUnitTest/NetcodeUnitTest/Source/NetcodeUnitTest/Private/UI/SLogWidget.cpp
+++ b/Engine/Plugins/NetcodeUnitTest/NetcodeUnitTest/Source/NetcodeUnitTest/Private/UI/SLogWidget.cpp
@@ -20,25 +20,6 @@
 #include "UI/SMultiSelectTableRow.h"
 
 
-// @todo #JohnBFeatureUI: It would be good to have both a 'Search' filter tab, and a 'Ctrl+F Find' button for each existing tab;
-//				two different search abilities (should be easy to add too)
-
-// @todo #JohnBFeatureUI: Could change the 'close' button, to a '+' button, which clones the current tab and enables search filtering
-
-
-// @todo #JohnBReview: Perhaps have conditional tabs, which only appear if log entries of this type come up?
-//				(such as: Warning, Error, Debug) UPDATE: Is this not already present?
-
-// @todo #JohnBFeatureUI: Consider adding a "TabName (LineCount)" to the tab name (maybe just for important tabs, like Error/Warning)
-
-
-// @todo #JohnBUI: Change the opacity, of background highlight for selected log entries - really ugly at the moment
-//					(unfortunately, this requires creating a new FTableRowStyle, which is a bit non-trivial)
-
-// @todo #JohnBUI: The 'failed to find text' message does not fade out properly, and hijacking ComputeVolatility does not help.
-//					Find a way to fix this, eventually.
-
-
 // Enable access to the private SEditableTextBox.EditableText variable, using the GET_PRIVATE macro
 IMPLEMENT_GET_PRIVATE_VAR(SEditableTextBox, EditableText, TSharedPtr<SEditableText>);
 
@@ -52,8 +33,6 @@ IMPLEMENT_GET_PROTECTED_FUNC_CONST(SBorder, GetBorderBackgroundColor, FSlateColo
 IMPLEMENT_GET_PROTECTED_FUNC_CONST(SDockTab, GetCurrentStyle, const FDockTabStyle&, void,, const);
 
 
-// @todo #JohnBRefactorUI: Perhaps move widget searching to a NUTSlate.h file, or such?
-
 /**
  * Delegate used for recursively iterating a widgets child widgets, and testing if they match a search condition
  *
@@ -419,7 +398,6 @@ void SLogWidget::Construct(const FArguments& Args)
 						SAssignNew(ConsoleComboBox, SComboBox<TSharedPtr<FString>>)
 						.OptionsSource(&ConsoleContextList)
 						.ToolTipText(FText::FromString(FString(TEXT("Select the context for executing console commands."))))
-						// @todo #JohnBRefactorUI: Too big to inline? Separate into its own function?
 						.OnGenerateWidget_Lambda(
 								[](TSharedPtr<FString> Item)
 								{
@@ -479,7 +457,6 @@ void SLogWidget::Construct(const FArguments& Args)
 						.HintText(FText::FromString(TEXT("Console")))
 						.ToolTipText(FText::FromString(TEXT("Executes a console command within the specified context.")))
 						.ClearKeyboardFocusOnCommit(false)
-						// @todo #JohnBRefactorUI: Too big to inline? Separate into its own function?
 						.OnTextCommitted_Lambda(
 							[&](const FText& InText, ETextCommit::Type InCommitType)
 							{
@@ -626,8 +603,6 @@ TSharedRef<SDockTab> SLogWidget::SpawnLogTab(const FSpawnTabArgs& InSpawnTabArgs
 			return CurElement->TabIdName == CurTabName;
 		});
 
-
-	// @todo #JohnBRefactorUI: Code duplication - this is defined above too - move to an inline function
 	auto ArrayAddNew =
 		[] (TArray<TSharedPtr<SWidget>>& InArray) -> TSharedPtr<SWidget>&
 		{
@@ -675,7 +650,6 @@ TSharedRef<SDockTab> SLogWidget::SpawnLogTab(const FSpawnTabArgs& InSpawnTabArgs
 						 */
 						SAssignNew(CurTabInfo->LogListView, SListView<TSharedRef<FLogLine>>)
 						.ListItemsSource(&CurTabInfo->TabLogLines)
-						// @todo #JohnBRefactorUI: Probably large enough to separate to its own function now
 						.OnGenerateRow_Lambda(
 							[](TSharedRef<FLogLine> Item, const TSharedRef<STableViewBase>& OwnerTable)
 							{
@@ -1054,7 +1028,7 @@ void SLogWidget::AddLine(ELogType InLogType, TSharedRef<FString> LogLine, FSlate
 {
 	TSharedRef<FLogLine> CurLogEntry = MakeShareable(new FLogLine(InLogType, LogLine, LogColor));
 
-	// Add the line to the master list
+	// Add the line to the main list
 	LogLines.Add(CurLogEntry);
 
 	TSharedPtr<FLogTabInfo> ActiveTab = GetActiveTabInfo();
diff --git a/Engine/Plugins/NetcodeUnitTest/NetcodeUnitTest/Source/NetcodeUnitTest/Private/UnitTests/Engine/SteamClient.cpp b/Engine/Plugins/NetcodeUnitTest/NetcodeUnitTest/Source/NetcodeUnitTest/Private/UnitTests/Engine/SteamClient.cpp
--- a/Engine/Plugins/NetcodeUnitTest/NetcodeUnitTest/Source/NetcodeUnitTest/Private/UnitTests/Engine/SteamClient.cpp
+++ b/Engine/Plugins/NetcodeUnitTest/NetcodeUnitTest/Source/NetcodeUnitTest/Private/UnitTests/Engine/SteamClient.cpp
@@ -97,7 +97,7 @@ void USteamClient::NotifyProcessLog(TWeakPtr<FUnitTestProcess> InProcess, const
 {
 	Super::NotifyProcessLog(InProcess, InLogLines);
 
-	const TCHAR* ServerAddrLog = TEXT("LogOnline: Verbose: STEAM: Master Server Data (P2PADDR, ");
+	const TCHAR* ServerAddrLog = TEXT(" Data (P2PADDR, ");
 
 	if (InProcess.HasSameObject(ServerHandle.Pin().Get()))
 	{
