Fixed handcrafted NMT_Join bunch.

#rnx

diff --git a/Engine/Plugins/NetcodeUnitTest/NetcodeUnitTest/Source/NetcodeUnitTest/Private/MinimalClient.cpp b/Engine/Plugins/NetcodeUnitTest/NetcodeUnitTest/Source/NetcodeUnitTest/Private/MinimalClient.cpp
--- a/Engine/Plugins/NetcodeUnitTest/NetcodeUnitTest/Source/NetcodeUnitTest/Private/MinimalClient.cpp
+++ b/Engine/Plugins/NetcodeUnitTest/NetcodeUnitTest/Source/NetcodeUnitTest/Private/MinimalClient.cpp
@@ -858,6 +858,13 @@ void UMinimalClient::SendInitialJoin()
 
 		*ControlChanBunch << EncryptionToken;
 
+		EEngineNetworkRuntimeFeatures LocalNetworkFeatures = EEngineNetworkRuntimeFeatures::None;
+		if (UNetDriver* NetDriver = UnitNetDriver.Get())
+		{
+			LocalNetworkFeatures = NetDriver->GetNetworkRuntimeFeatures();
+		}
+		*ControlChanBunch << LocalNetworkFeatures;
+
 		bool bSkipControlJoin = !!(MinClientFlags & EMinClientFlags::SkipControlJoin);
 		bool bBeaconConnect = !!(MinClientFlags & EMinClientFlags::BeaconConnect);
 
