Enable truncation warnings in NetcodeUnitTest.

diff --git a/Engine/Plugins/NetcodeUnitTest/NUTUnrealEngine/Source/NUTUnrealEngine/NUTUnrealEngine.Build.cs b/Engine/Plugins/NetcodeUnitTest/NUTUnrealEngine/Source/NUTUnrealEngine/NUTUnrealEngine.Build.cs
--- a/Engine/Plugins/NetcodeUnitTest/NUTUnrealEngine/Source/NUTUnrealEngine/NUTUnrealEngine.Build.cs
+++ b/Engine/Plugins/NetcodeUnitTest/NUTUnrealEngine/Source/NUTUnrealEngine/NUTUnrealEngine.Build.cs
@@ -16,6 +16,8 @@ namespace UnrealBuildTool.Rules
 					"NetcodeUnitTest"
 				}
 			);
+
+			UnsafeTypeCastWarningLevel = WarningLevel.Error;
 		}
 	}
 }
diff --git a/Engine/Plugins/NetcodeUnitTest/NUTUnrealEngine/Source/NUTUnrealEngine/Private/UnitTests/NetBitsTest.cpp b/Engine/Plugins/NetcodeUnitTest/NUTUnrealEngine/Source/NUTUnrealEngine/Private/UnitTests/NetBitsTest.cpp
--- a/Engine/Plugins/NetcodeUnitTest/NUTUnrealEngine/Source/NUTUnrealEngine/Private/UnitTests/NetBitsTest.cpp
+++ b/Engine/Plugins/NetcodeUnitTest/NUTUnrealEngine/Source/NUTUnrealEngine/Private/UnitTests/NetBitsTest.cpp
@@ -302,7 +302,7 @@ bool UNetBitsTest::ExecuteUnitTest()
 
 	// @todo #JohnB: REmove or incorporate properly - checking MAX_PACKETID size limits
 	{
-		uint32 MaxBits = 0;
+		int64 MaxBits = 0;
 
 		for (uint32 WriteValue=0; WriteValue<MAX_PACKETID; WriteValue++)
 		{
diff --git a/Engine/Plugins/NetcodeUnitTest/NetcodeUnitTest/Source/NetcodeUnitTest/Classes/NUTActor.h b/Engine/Plugins/NetcodeUnitTest/NetcodeUnitTest/Source/NetcodeUnitTest/Classes/NUTActor.h
--- a/Engine/Plugins/NetcodeUnitTest/NetcodeUnitTest/Source/NetcodeUnitTest/Classes/NUTActor.h
+++ b/Engine/Plugins/NetcodeUnitTest/NetcodeUnitTest/Source/NetcodeUnitTest/Classes/NUTActor.h
@@ -56,7 +56,7 @@ private:
 
 public:
 	/** The value of World.RealTimeSeconds as of the last time the client was marked as still alive */
-	float LastAliveTime;
+	double LastAliveTime;
 
 	/** A delegate property, used solely for converting strings to delegates */
 	UPROPERTY()
diff --git a/Engine/Plugins/NetcodeUnitTest/NetcodeUnitTest/Source/NetcodeUnitTest/NetcodeUnitTest.Build.cs b/Engine/Plugins/NetcodeUnitTest/NetcodeUnitTest/Source/NetcodeUnitTest/NetcodeUnitTest.Build.cs
--- a/Engine/Plugins/NetcodeUnitTest/NetcodeUnitTest/Source/NetcodeUnitTest/NetcodeUnitTest.Build.cs
+++ b/Engine/Plugins/NetcodeUnitTest/NetcodeUnitTest/Source/NetcodeUnitTest/NetcodeUnitTest.Build.cs
@@ -42,6 +42,8 @@ namespace UnrealBuildTool.Rules
 					}
 				);
 			}
+
+			UnsafeTypeCastWarningLevel = WarningLevel.Error;
 		}
 	}
 }
diff --git a/Engine/Plugins/NetcodeUnitTest/NetcodeUnitTest/Source/NetcodeUnitTest/Private/ClientUnitTest.cpp b/Engine/Plugins/NetcodeUnitTest/NetcodeUnitTest/Source/NetcodeUnitTest/Private/ClientUnitTest.cpp
--- a/Engine/Plugins/NetcodeUnitTest/NetcodeUnitTest/Source/NetcodeUnitTest/Private/ClientUnitTest.cpp
+++ b/Engine/Plugins/NetcodeUnitTest/NetcodeUnitTest/Source/NetcodeUnitTest/Private/ClientUnitTest.cpp
@@ -108,7 +108,7 @@ void UClientUnitTest::NotifyControlMessage(FInBunch& Bunch, uint8 MessageType)
 		if (!Bunch.IsError() && Bunch.GetBitsLeft() > 0)
 		{
 			UNIT_LOG_BEGIN(this, ELogType::StatusDebug | ELogType::StyleMonospace);
-			NUTDebug::LogHexDump(Bunch.GetDataPosChecked(), Bunch.GetBytesLeft(), true, true);
+			NUTDebug::LogHexDump(Bunch.GetDataPosChecked(), IntCastChecked<int32, int64>(Bunch.GetBytesLeft()), true, true);
 			UNIT_LOG_END();
 		}
 	}
diff --git a/Engine/Plugins/NetcodeUnitTest/NetcodeUnitTest/Source/NetcodeUnitTest/Private/MinimalClient.cpp b/Engine/Plugins/NetcodeUnitTest/NetcodeUnitTest/Source/NetcodeUnitTest/Private/MinimalClient.cpp
--- a/Engine/Plugins/NetcodeUnitTest/NetcodeUnitTest/Source/NetcodeUnitTest/Private/MinimalClient.cpp
+++ b/Engine/Plugins/NetcodeUnitTest/NetcodeUnitTest/Source/NetcodeUnitTest/Private/MinimalClient.cpp
@@ -258,13 +258,13 @@ bool UMinimalClient::SendRawBunch(FOutBunch& Bunch, bool bAllowPartial/*=false*/
 	{
 		TArray<FOutBunch*> SendBunches;
 		const int32 MAX_SINGLE_BUNCH_SIZE_BITS = UnitConn->GetMaxSingleBunchSizeBits();
-		int32 BunchNumBits = Bunch.GetNumBits();
+		int64 BunchNumBits = Bunch.GetNumBits();
 
 		if (bAllowPartial && BunchNumBits > MAX_SINGLE_BUNCH_SIZE_BITS)
 		{
 			const int32 MAX_SINGLE_BUNCH_SIZE_BYTES = MAX_SINGLE_BUNCH_SIZE_BITS / 8;
 			const int32 MAX_PARTIAL_BUNCH_SIZE_BITS = MAX_SINGLE_BUNCH_SIZE_BYTES * 8;
-			int32 NumSerializedBits = 0;
+			int64 NumSerializedBits = 0;
 
 			// Discard the original bunch before splitting it, in order to have the correct packet sequence for the new packets
 			DiscardChannelBunch(&Bunch);
@@ -272,7 +272,7 @@ bool UMinimalClient::SendRawBunch(FOutBunch& Bunch, bool bAllowPartial/*=false*/
 			while (NumSerializedBits < BunchNumBits)
 			{
 				FOutBunch* NewBunch = CreateChannelBunchByName(Bunch.ChName, Bunch.ChIndex);
-				int32 NewNumBits = FMath::Min(BunchNumBits - NumSerializedBits, MAX_PARTIAL_BUNCH_SIZE_BITS);
+				int64 NewNumBits = FMath::Min(BunchNumBits - NumSerializedBits, MAX_PARTIAL_BUNCH_SIZE_BITS);
 
 				NewBunch->SerializeBits(Bunch.GetData() + (NumSerializedBits >> 3), NewNumBits);
 				NumSerializedBits += NewNumBits;
diff --git a/Engine/Plugins/NetcodeUnitTest/NetcodeUnitTest/Source/NetcodeUnitTest/Private/NUTActor.cpp b/Engine/Plugins/NetcodeUnitTest/NetcodeUnitTest/Source/NetcodeUnitTest/Private/NUTActor.cpp
--- a/Engine/Plugins/NetcodeUnitTest/NetcodeUnitTest/Source/NetcodeUnitTest/Private/NUTActor.cpp
+++ b/Engine/Plugins/NetcodeUnitTest/NetcodeUnitTest/Source/NetcodeUnitTest/Private/NUTActor.cpp
@@ -460,8 +460,8 @@ void ANUTActor::HookNetDriver(UNetDriver* TargetNetDriver)
 			UE_LOG(LogUnitTest, Log, TEXT("Setting %s InitialConnectTimeout/ConnectionTimeout to '%i'"), *TargetNetDriver->GetFullName(),
 					CustomTimeout);
 
-			TargetNetDriver->InitialConnectTimeout = CustomTimeout;
-			TargetNetDriver->ConnectionTimeout = CustomTimeout;
+			TargetNetDriver->InitialConnectTimeout = static_cast<float>(CustomTimeout);
+			TargetNetDriver->ConnectionTimeout = static_cast<float>(CustomTimeout);
 		}
 	}
 }
diff --git a/Engine/Plugins/NetcodeUnitTest/NetcodeUnitTest/Source/NetcodeUnitTest/Private/NUTUtil.cpp b/Engine/Plugins/NetcodeUnitTest/NetcodeUnitTest/Source/NetcodeUnitTest/Private/NUTUtil.cpp
--- a/Engine/Plugins/NetcodeUnitTest/NetcodeUnitTest/Source/NetcodeUnitTest/Private/NUTUtil.cpp
+++ b/Engine/Plugins/NetcodeUnitTest/NetcodeUnitTest/Source/NetcodeUnitTest/Private/NUTUtil.cpp
@@ -108,7 +108,7 @@ void NUTUtil::SortUnitTestClassDefList(TArray<UUnitTest*>& InUnitTestClassDefaul
 	InUnitTestClassDefaults.Sort(FUnitTestTypeDateSort(ListTypes));
 }
 
-bool NUTUtil::ParseEnum(const TCHAR* Stream, const TCHAR* Match, const TCHAR* EnumName, uint32& Value)
+bool NUTUtil::ParseEnum(const TCHAR* Stream, const TCHAR* Match, const TCHAR* EnumName, uint64& Value)
 {
 	bool bReturnVal = false;
 	FString StrValue;
diff --git a/Engine/Plugins/NetcodeUnitTest/NetcodeUnitTest/Source/NetcodeUnitTest/Private/ProcessUnitTest.cpp b/Engine/Plugins/NetcodeUnitTest/NetcodeUnitTest/Source/NetcodeUnitTest/Private/ProcessUnitTest.cpp
--- a/Engine/Plugins/NetcodeUnitTest/NetcodeUnitTest/Source/NetcodeUnitTest/Private/ProcessUnitTest.cpp
+++ b/Engine/Plugins/NetcodeUnitTest/NetcodeUnitTest/Source/NetcodeUnitTest/Private/ProcessUnitTest.cpp
@@ -633,7 +633,7 @@ void UProcessUnitTest::PollProcessOutput()
 			{
 				if (MaxPipeDelay - PipeCounter > 0.0)
 				{
-					FPlatformProcess::SleepNoStats(CurPipeDelay);
+					FPlatformProcess::SleepNoStats(static_cast<float>(CurPipeDelay));
 
 					PipeCounter += CurPipeDelay;
 				}
diff --git a/Engine/Plugins/NetcodeUnitTest/NetcodeUnitTest/Source/NetcodeUnitTest/Private/UI/LogWindowManager.cpp b/Engine/Plugins/NetcodeUnitTest/NetcodeUnitTest/Source/NetcodeUnitTest/Private/UI/LogWindowManager.cpp
--- a/Engine/Plugins/NetcodeUnitTest/NetcodeUnitTest/Source/NetcodeUnitTest/Private/UI/LogWindowManager.cpp
+++ b/Engine/Plugins/NetcodeUnitTest/NetcodeUnitTest/Source/NetcodeUnitTest/Private/UI/LogWindowManager.cpp
@@ -51,8 +51,8 @@ void FLogWindowManager::Initialize(int InLogWidth, int InLogHeight)
 		FPlatformRect WorkAreaRect = DispMetrics.PrimaryDisplayWorkAreaRect;
 
 		// Calculate how many windows will fit horizontally, and then vertically
-		float HorizontalCount = FMath::FloorToFloat((WorkAreaRect.Right - WorkAreaRect.Left) / (float)LogWidth);
-		float VerticalCount = FMath::FloorToFloat((WorkAreaRect.Bottom - WorkAreaRect.Top) / (float)LogHeight);
+		int HorizontalCount = FMath::FloorToInt(static_cast<float>((WorkAreaRect.Right - WorkAreaRect.Left) / LogWidth));
+		int VerticalCount = FMath::FloorToInt(static_cast<float>((WorkAreaRect.Bottom - WorkAreaRect.Top) / LogHeight));
 
 		// Now setup the grid
 		for (int VIdx = 0; VIdx<VerticalCount; VIdx++)
@@ -62,13 +62,13 @@ void FLogWindowManager::Initialize(int InLogWidth, int InLogHeight)
 				FLogGridEntry CurEntry;
 
 				// Calculate the current offset of the log window, for this grid entry
-				float VerticalOffset = (float)LogHeight * (float)VIdx;
-				float HorizontalOffset = (float)LogWidth * (float)HIdx;
+				int VerticalOffset = LogHeight * VIdx;
+				int HorizontalOffset = LogWidth * HIdx;
 
-				CurEntry.Top = WorkAreaRect.Top + VerticalOffset;
-				CurEntry.Bottom = WorkAreaRect.Top + VerticalOffset + LogHeight;
-				CurEntry.Left = WorkAreaRect.Left + HorizontalOffset;
-				CurEntry.Right = WorkAreaRect.Left + HorizontalOffset + LogWidth;
+				CurEntry.Top = static_cast<float>(WorkAreaRect.Top + VerticalOffset);
+				CurEntry.Bottom = static_cast<float>(WorkAreaRect.Top + VerticalOffset + LogHeight);
+				CurEntry.Left = static_cast<float>(WorkAreaRect.Left + HorizontalOffset);
+				CurEntry.Right = static_cast<float>(WorkAreaRect.Left + HorizontalOffset + LogWidth);
 
 				GridSpaces.Add(CurEntry);
 			}
@@ -102,7 +102,7 @@ TSharedPtr<SLogWindow> FLogWindowManager::CreateLogWindow(FString Title, ELogTyp
 		FLogGridEntry& CurEntry = GridSpaces[FreeGridPos];
 
 		ReturnVal =
-				SNew(SLogWindow, Title, CurEntry.Left, CurEntry.Top, LogWidth, LogHeight)
+				SNew(SLogWindow, Title, CurEntry.Left, CurEntry.Top, static_cast<float>(LogWidth), static_cast<float>(LogHeight))
 				.bStatusWindow(bStatusWindow)
 				.ExpectedFilters(ExpectedFilters);
 
@@ -117,7 +117,7 @@ TSharedPtr<SLogWindow> FLogWindowManager::CreateLogWindow(FString Title, ELogTyp
 		FLogGridEntry& OverflowLoc = GridSpaces[0];
 
 		ReturnVal =
-				SNew(SLogWindow, Title, OverflowLoc.Left, OverflowLoc.Top, LogWidth, LogHeight)
+				SNew(SLogWindow, Title, OverflowLoc.Left, OverflowLoc.Top, static_cast<float>(LogWidth), static_cast<float>(LogHeight))
 				.bStatusWindow(bStatusWindow)
 				.ExpectedFilters(ExpectedFilters);
 
diff --git a/Engine/Plugins/NetcodeUnitTest/NetcodeUnitTest/Source/NetcodeUnitTest/Private/UI/SLogWidget.cpp b/Engine/Plugins/NetcodeUnitTest/NetcodeUnitTest/Source/NetcodeUnitTest/Private/UI/SLogWidget.cpp
--- a/Engine/Plugins/NetcodeUnitTest/NetcodeUnitTest/Source/NetcodeUnitTest/Private/UI/SLogWidget.cpp
+++ b/Engine/Plugins/NetcodeUnitTest/NetcodeUnitTest/Source/NetcodeUnitTest/Private/UI/SLogWidget.cpp
@@ -696,10 +696,10 @@ TSharedRef<SDockTab> SLogWidget::SpawnLogTab(const FSpawnTabArgs& InSpawnTabArgs
 									const TSharedRef<FSlateFontMeasure> FontMeasure =
 																		FSlateApplication::Get().GetRenderer()->GetFontMeasureService();
 
-									FVector2D UnderlineDim = FontMeasure->Measure(TextToUnderline, 0, TotalLen, RenderFont);
-									FVector2D BaseDim = FontMeasure->Measure(TEXT("-"), RenderFont);
+									FVector2f UnderlineDim = FontMeasure->Measure(TextToUnderline, 0, TotalLen, RenderFont);
+									FVector2f BaseDim = FontMeasure->Measure(TEXT("-"), RenderFont);
 
-									uint32 UnderlineCharCount = FMath::FloorToInt(UnderlineDim.X / BaseDim.X);
+									int32 UnderlineCharCount = FMath::FloorToInt(UnderlineDim.X / BaseDim.X);
 
 									if (UnderlineCharCount > 0)
 									{
@@ -775,7 +775,7 @@ TSharedRef<SDockTab> SLogWidget::SpawnLogTab(const FSpawnTabArgs& InSpawnTabArgs
 									.Visibility(EVisibility::Collapsed)
 									.ToolTipText(FText::FromString(TEXT("Close the find bar.")))
 									.ButtonStyle(CloseButtonStyle)
-									.ContentPadding(0)
+									.ContentPadding(0.0f)
 									.OnClicked_Lambda(
 										[&]()
 										{
diff --git a/Engine/Plugins/NetcodeUnitTest/NetcodeUnitTest/Source/NetcodeUnitTest/Private/UnitTestManager.cpp b/Engine/Plugins/NetcodeUnitTest/NetcodeUnitTest/Source/NetcodeUnitTest/Private/UnitTestManager.cpp
--- a/Engine/Plugins/NetcodeUnitTest/NetcodeUnitTest/Source/NetcodeUnitTest/Private/UnitTestManager.cpp
+++ b/Engine/Plugins/NetcodeUnitTest/NetcodeUnitTest/Source/NetcodeUnitTest/Private/UnitTestManager.cpp
@@ -1554,7 +1554,7 @@ bool UUnitTestManager::Exec(UWorld* InWorld, const TCHAR* Cmd, FOutputDevice& Ar
 		{
 #if STATS
 			FString TargetEvent = FParse::Token(Cmd, false);
-			uint8 FramePercentThreshold = FCString::Atoi(*FParse::Token(Cmd, false));
+			int32 FramePercentThreshold = FCString::Atoi(*FParse::Token(Cmd, false));
 
 			if (!TargetEvent.IsEmpty() && FramePercentThreshold > 0)
 			{
diff --git a/Engine/Plugins/NetcodeUnitTest/NetcodeUnitTest/Source/NetcodeUnitTest/Private/UnitTests/Engine/PackedVectorTest.cpp b/Engine/Plugins/NetcodeUnitTest/NetcodeUnitTest/Source/NetcodeUnitTest/Private/UnitTests/Engine/PackedVectorTest.cpp
--- a/Engine/Plugins/NetcodeUnitTest/NetcodeUnitTest/Source/NetcodeUnitTest/Private/UnitTests/Engine/PackedVectorTest.cpp
+++ b/Engine/Plugins/NetcodeUnitTest/NetcodeUnitTest/Source/NetcodeUnitTest/Private/UnitTests/Engine/PackedVectorTest.cpp
@@ -151,7 +151,7 @@ bool UPackedVectorTest::ExecuteFloatTest(TMap<FString, bool>& TestResults)
 					else
 					{
 						const float ValueDiff = FMath::Abs(ReadValue.X - WriteValue.X);
-						LocalSuccess &= ValueDiff < 2.0f/Test.ScaleFactor; // The diff test might need some adjustment
+						LocalSuccess &= ValueDiff < 2.0f/static_cast<float>(Test.ScaleFactor); // The diff test might need some adjustment
 					}
 				}
 			}
diff --git a/Engine/Plugins/NetcodeUnitTest/NetcodeUnitTest/Source/NetcodeUnitTest/Public/NUTUtil.h b/Engine/Plugins/NetcodeUnitTest/NetcodeUnitTest/Source/NetcodeUnitTest/Public/NUTUtil.h
--- a/Engine/Plugins/NetcodeUnitTest/NetcodeUnitTest/Source/NetcodeUnitTest/Public/NUTUtil.h
+++ b/Engine/Plugins/NetcodeUnitTest/NetcodeUnitTest/Source/NetcodeUnitTest/Public/NUTUtil.h
@@ -433,7 +433,7 @@ struct NUTUtil
 			bool bArgumentsQuoted = *Start == '"';
 
 			// Number of characters we can look back from found looking for first parenthesis.
-			uint32 AllowedBacktraceCharactersCount = Found - Stream;
+			uint64 AllowedBacktraceCharactersCount = Found - Stream;
 
 			// Check for fully quoted string with spaces
 			bool bFullyQuoted = 
@@ -572,7 +572,7 @@ struct NUTUtil
 	 * @param Value		Outputs the integer value of the enum, as specified on the commandline
 	 * @return			Whether or not an enum value was parsed successfully
 	 */
-	static bool ParseEnum(const TCHAR* Stream, const TCHAR* Match, const TCHAR* EnumName, uint32& Value);
+	static bool ParseEnum(const TCHAR* Stream, const TCHAR* Match, const TCHAR* EnumName, uint64& Value);
 
 
 	/**
diff --git a/Engine/Plugins/NetcodeUnitTest/NetcodeUnitTest/Source/NetcodeUnitTest/Public/NUTUtilProfiler.h b/Engine/Plugins/NetcodeUnitTest/NetcodeUnitTest/Source/NetcodeUnitTest/Public/NUTUtilProfiler.h
--- a/Engine/Plugins/NetcodeUnitTest/NetcodeUnitTest/Source/NetcodeUnitTest/Public/NUTUtilProfiler.h
+++ b/Engine/Plugins/NetcodeUnitTest/NetcodeUnitTest/Source/NetcodeUnitTest/Public/NUTUtilProfiler.h
@@ -25,7 +25,7 @@ public:
 	 * @param InTargetEvent				The event to be targeted/probed
 	 * @param InFramePercentThreshold	The percentage of frame time used by the event, needed to trigger detection
 	 */
-	FFrameProfiler(FName InTargetEvent, uint8 InFramePercentThreshold)
+	FFrameProfiler(FName InTargetEvent, int32 InFramePercentThreshold)
 		: TargetEvent(InTargetEvent)
 		, FramePercentThreshold(InFramePercentThreshold)
 		, bActive(false)
@@ -73,7 +73,7 @@ public:
 	FName TargetEvent;
 
 	/** The percentage of frame time used by the event needed to trigger detection */
-	uint8 FramePercentThreshold;
+	int32 FramePercentThreshold;
 
 
 protected:
