Temporary workaround for a CIS error.

diff --git a/Engine/Plugins/NetcodeUnitTest/NetcodeUnitTest/Source/NetcodeUnitTest/Private/NUTUtilReflectionParser.cpp b/Engine/Plugins/NetcodeUnitTest/NetcodeUnitTest/Source/NetcodeUnitTest/Private/NUTUtilReflectionParser.cpp
--- a/Engine/Plugins/NetcodeUnitTest/NetcodeUnitTest/Source/NetcodeUnitTest/Private/NUTUtilReflectionParser.cpp
+++ b/Engine/Plugins/NetcodeUnitTest/NetcodeUnitTest/Source/NetcodeUnitTest/Private/NUTUtilReflectionParser.cpp
@@ -196,7 +196,7 @@ FVMReflectionParser::FVMReflectionParser()
 
 		if (RealContext != nullptr)
 		{
-			FVMReflection& Refl = *(RealContext->Refl);
+			FVMReflection& Refl = *(RealContext->Refl.Get());
 
 			Refl->*(B.IdentifierName);
 
