Initialized some FVector variables that were potentially uninitialized.

diff --git a/Engine/Plugins/NetcodeUnitTest/NetcodeUnitTest/Source/NetcodeUnitTest/Private/UnitTests/Engine/PackedVectorTest.cpp b/Engine/Plugins/NetcodeUnitTest/NetcodeUnitTest/Source/NetcodeUnitTest/Private/UnitTests/Engine/PackedVectorTest.cpp
--- a/Engine/Plugins/NetcodeUnitTest/NetcodeUnitTest/Source/NetcodeUnitTest/Private/UnitTests/Engine/PackedVectorTest.cpp
+++ b/Engine/Plugins/NetcodeUnitTest/NetcodeUnitTest/Source/NetcodeUnitTest/Private/UnitTests/Engine/PackedVectorTest.cpp
@@ -122,7 +122,7 @@ bool UPackedVectorTest::ExecuteFloatTest(TMap<FString, bool>& TestResults)
 
 			const float ScalarValue = Test.TestValues[ValueIt];
 			const FVector3f WriteValue(ScalarValue);
-			FVector3f ReadValue;
+			FVector3f ReadValue = FVector3f::ZeroVector;
 
 			const bool bOverflowOrNan = !Test.Writer(WriteValue, Writer);
 			bool LocalSuccess = !Writer.GetError();
@@ -230,7 +230,7 @@ bool UPackedVectorTest::ExecuteDoubleTest(TMap<FString, bool>& TestResults)
 
 			const double ScalarValue = Test.TestValues[ValueIt];
 			const FVector3d WriteValue(ScalarValue);
-			FVector3d ReadValue;
+			FVector3d ReadValue = FVector3d::ZeroVector;
 
 			const bool bOverflowOrNan = !Test.Writer(WriteValue, Writer);
 			bool LocalSuccess = !Writer.GetError();
diff --git a/Engine/Source/Runtime/Experimental/Chaos/Public/Chaos/Tetrahedron.h b/Engine/Source/Runtime/Experimental/Chaos/Public/Chaos/Tetrahedron.h
--- a/Engine/Source/Runtime/Experimental/Chaos/Public/Chaos/Tetrahedron.h
+++ b/Engine/Source/Runtime/Experimental/Chaos/Public/Chaos/Tetrahedron.h
@@ -353,7 +353,7 @@ namespace Chaos {
 				TVec2<int32>(1,3),
 				TVec2<int32>(2,3)
 			};
-			TVec3<T> ClosestEdgePoint;
+			TVec3<T> ClosestEdgePoint = TVec3<T>::ZeroVector;
 			T ClosestEdgeAlpha(0);
 			int32 ClosestEdgeIndex = INDEX_NONE;
 			T ClosestEdgeDistSq = std::numeric_limits<T>::max();
diff --git a/Engine/Source/Runtime/PhysicsCore/Public/SQVisitor.h b/Engine/Source/Runtime/PhysicsCore/Public/SQVisitor.h
--- a/Engine/Source/Runtime/PhysicsCore/Public/SQVisitor.h
+++ b/Engine/Source/Runtime/PhysicsCore/Public/SQVisitor.h
@@ -586,9 +586,9 @@ private:
 		return bContinue;
 	}
 
-	const FVector StartPoint;
-	const FVector Dir;
-	const FVector HalfExtents;
+	const FVector StartPoint = FVector::ZeroVector;
+	const FVector Dir = FVector::ZeroVector;
+	const FVector HalfExtents = FVector::ZeroVector;
 	FHitFlags OutputFlags;
 	bool bAnyHit;
 	const ChaosInterface::FQueryDebugParams DebugParams;
