Fix UnitTestManager
#rnx

diff --git a/Engine/Plugins/NetcodeUnitTest/NetcodeUnitTest/Source/NetcodeUnitTest/Private/UnitTestManager.cpp b/Engine/Plugins/NetcodeUnitTest/NetcodeUnitTest/Source/NetcodeUnitTest/Private/UnitTestManager.cpp
--- a/Engine/Plugins/NetcodeUnitTest/NetcodeUnitTest/Source/NetcodeUnitTest/Private/UnitTestManager.cpp
+++ b/Engine/Plugins/NetcodeUnitTest/NetcodeUnitTest/Source/NetcodeUnitTest/Private/UnitTestManager.cpp
@@ -2004,11 +2004,14 @@ void UUnitTestManager::Serialize(const TCHAR* Data, ELogVerbosity::Type Verbosit
 
 					if (TimerData != nullptr && TimerData->TimerDelegate.IsBound())
 					{
-						UObject* TimerObj = TimerData->TimerDelegate.FuncDelegate.GetUObject();
-
-						if (TimerObj == nullptr)
+						UObject* TimerObj = nullptr;
+						if (FTimerDelegate* TimerDelegatge = TimerData->TimerDelegate.VariantDelegate.TryGet<FTimerDelegate>())
+						{
+							TimerObj = TimerDelegatge->GetUObject();
+						}
+						else if (FTimerDynamicDelegate* TimerDynDelegatge = TimerData->TimerDelegate.VariantDelegate.TryGet<FTimerDynamicDelegate>())
 						{
-							TimerObj = TimerData->TimerDelegate.FuncDynDelegate.GetUObject();
+							TimerObj = TimerDynDelegatge->GetUObject();
 						}
 
 						if (TimerObj != nullptr)
