Fixed non unity compile errors because of missing StrProperty.h include

diff --git a/Engine/Plugins/NetcodeUnitTest/NetcodeUnitTest/Source/NetcodeUnitTest/Private/NUTUtilReflection.cpp b/Engine/Plugins/NetcodeUnitTest/NetcodeUnitTest/Source/NetcodeUnitTest/Private/NUTUtilReflection.cpp
--- a/Engine/Plugins/NetcodeUnitTest/NetcodeUnitTest/Source/NetcodeUnitTest/Private/NUTUtilReflection.cpp
+++ b/Engine/Plugins/NetcodeUnitTest/NetcodeUnitTest/Source/NetcodeUnitTest/Private/NUTUtilReflection.cpp
@@ -5,6 +5,7 @@
 #include "Misc/TVariantMeta.h"
 #include "UObject/EnumProperty.h"
 #include "Templates/ValueOrError.h"
+#include "UObject/StrProperty.h"
 #include "UObject/TextProperty.h"
 #include "NetcodeUnitTest.h"
 
diff --git a/Engine/Source/Editor/SerializedRecorderInterface/Public/Serializers/MovieScenePropertySerialization.h b/Engine/Source/Editor/SerializedRecorderInterface/Public/Serializers/MovieScenePropertySerialization.h
--- a/Engine/Source/Editor/SerializedRecorderInterface/Public/Serializers/MovieScenePropertySerialization.h
+++ b/Engine/Source/Editor/SerializedRecorderInterface/Public/Serializers/MovieScenePropertySerialization.h
@@ -3,6 +3,7 @@
 #pragma once
 
 #include "Misc/FrameRate.h"
+#include "UObject/StrProperty.h"
 #include "UObject/UnrealType.h"
 #include "UObject/EnumProperty.h"
 #include "MovieSceneSectionSerialization.h"
